Copying //UE4/Dev-Sequencer to //UE4/Main (Source: //UE4/Dev-Sequencer @ 3617642)

diff --git a/Engine/Plugins/Media/AndroidCamera/Source/AndroidCamera/Java/src/com/epicgames/ue4/CameraPlayer14.java b/Engine/Plugins/Media/AndroidCamera/Source/AndroidCamera/Java/src/com/epicgames/ue4/CameraPlayer14.java
new file mode 100644
--- /dev/null
+++ b/Engine/Plugins/Media/AndroidCamera/Source/AndroidCamera/Java/src/com/epicgames/ue4/CameraPlayer14.java
@@ -0,0 +1,2125 @@
+// Copyright 1998-2017 Epic Games, Inc. All Rights Reserved.
+
+package com.epicgames.ue4;
+
+import android.app.Activity;
+import android.app.Application;
+import android.content.res.AssetFileDescriptor;
+import android.content.res.AssetManager;
+import android.content.pm.PackageManager;
+import android.content.res.Configuration;
+import android.graphics.Rect;
+import android.opengl.*;
+import android.os.Build;
+import android.os.Bundle;
+import android.view.Surface;
+import java.io.File;
+import java.io.IOException;
+import java.io.FileOutputStream;
+import java.io.FileNotFoundException;
+import java.io.RandomAccessFile;
+import java.lang.RuntimeException;
+import android.hardware.Camera;
+import java.util.List;
+import java.util.ArrayList;
+import java.util.Random;
+import android.util.Log;
+
+// from permission_library
+import com.google.vr.sdk.samples.permission.PermissionHelper;
+
+/*
+	Custom camera player that renders video to a frame buffer. This
+	variation is for API 14 and above.
+*/
+public class CameraPlayer14
+{
+	private enum CameraStates {
+		PREINIT, INIT, PREPARED, START, PAUSE, STOP, RELEASED
+	}
+
+	private boolean CameraSupported = false;
+	private boolean mVideoEnabled = true;
+	
+	private Camera mCamera = null;
+	private CameraStates CameraState = CameraStates.PREINIT;
+	private int CameraId = 0;
+	private int CameraWidth = 0;
+	private int CameraHeight = 0;
+	private int CameraFPS = 0;
+	private int CameraFPSMin = 0;
+	private int CameraFPSMax = 0;
+	private int CameraOrientation = 0;
+	private int DeviceRotation = 0;
+	
+	private int iFrontId = 0;
+	private int iBackId = 0;
+
+	private boolean SwizzlePixels = true;
+	private boolean VulkanRenderer = false;
+	private boolean Looping = false;
+	private volatile boolean WaitOnBitmapRender = false;
+	private volatile boolean Prepared = false;
+	private volatile boolean Completed = false;
+
+	private BitmapRenderer mBitmapRenderer = null;
+	private OESTextureRenderer mOESTextureRenderer = null;
+
+	public class FrameUpdateInfo {
+		java.nio.Buffer Buffer;
+		int CurrentPosition;
+		boolean FrameReady;
+		boolean RegionChanged;
+		float ScaleRotation00;
+		float ScaleRotation01;
+		float ScaleRotation10;
+		float ScaleRotation11;
+		float UOffset;
+		float VOffset;
+	}
+
+	public class AudioTrackInfo {
+		public int Index;
+		public String MimeType;
+		public String DisplayName;
+		public String Language;
+		public int Channels;
+		public int SampleRate;
+	}
+
+	public class CaptionTrackInfo {
+		public int Index;
+		public String MimeType;
+		public String DisplayName;
+		public String Language;
+	}
+
+	public class VideoTrackInfo {
+		public int Index;
+		public String MimeType;
+		public String DisplayName;
+		public String Language;
+		public int BitRate;
+		public int Width;
+		public int Height;
+		public float FrameRate;
+		public float FrameRateLow;
+		public float FrameRateHigh;
+	}
+
+	private ArrayList<AudioTrackInfo> audioTracks = new ArrayList<AudioTrackInfo>();
+	private ArrayList<VideoTrackInfo> videoTracks = new ArrayList<VideoTrackInfo>();
+
+	// ======================================================================================
+
+	public CameraPlayer14(boolean swizzlePixels, boolean vulkanRenderer)
+	{
+		CameraState = CameraStates.INIT;
+		SwizzlePixels = swizzlePixels;
+		VulkanRenderer = vulkanRenderer;
+		WaitOnBitmapRender = false;
+
+		/*
+		setOnErrorListener(new MediaPlayer.OnErrorListener() {
+			@Override
+			public boolean onError(MediaPlayer mp, int what, int extra) {
+				GameActivity.Log.debug("MediaPlayer14: onError returned what=" + what + ", extra=" + extra);
+				return true;
+			}
+		});
+		*/
+
+		// check if camera supported
+		if (GameActivity._activity.getPackageManager().hasSystemFeature(PackageManager.FEATURE_CAMERA_ANY))
+		{
+			CameraSupported = true;
+			GameActivity.Log.debug("Camera supported");
+		}
+		else
+		{
+			CameraSupported = false;
+			GameActivity.Log.debug("Camera not supported");
+			return;
+		}
+
+		// make sure we have permission to use camera (user should request this first!)
+		if (!PermissionHelper.checkPermission("android.permission.CAMERA"))
+		{
+			// don't have permission so disable
+			CameraSupported = false;
+			GameActivity.Log.debug("Camera permission not granted");
+			return;
+		}
+		
+		// discover available cameras and remember front and back IDs if available
+		Camera.CameraInfo cameraInfo = new Camera.CameraInfo();
+		int cameraCount = Camera.getNumberOfCameras();
+		for (int index = 0; index < cameraCount; index++) {
+			Camera.getCameraInfo(index, cameraInfo);
+			if (cameraInfo.facing == Camera.CameraInfo.CAMERA_FACING_FRONT)
+			{
+				iFrontId = index;
+			} else if (cameraInfo.facing == Camera.CameraInfo.CAMERA_FACING_BACK)
+			{
+				iBackId = index;
+			}
+		}
+
+		GameActivity.Log.debug("CameraPlayer14: frontId=" + iFrontId + ", backId=" + iBackId);
+
+		// not really necessary, but just in case...
+		if (Build.VERSION.SDK_INT >= Build.VERSION_CODES.ICE_CREAM_SANDWICH)
+		{
+			// register lifecycle callbacks to unlock / reconnect to camera if application goes to background
+			GameActivity._activity.getApplication().registerActivityLifecycleCallbacks(
+				new Application.ActivityLifecycleCallbacks() {
+					@Override
+					public void onActivityCreated(Activity activity, Bundle savedInstanceState) {}
+					@Override
+					public void onActivityStarted(Activity activity) {}
+					@Override
+					public void onActivitySaveInstanceState(Activity activity, Bundle outState) {}
+					@Override
+					public void onActivityPaused(Activity activity)
+					{
+						if (null != mCamera)
+						{
+							try {
+								if (CameraState == CameraStates.START)
+								{
+									GameActivity.Log.debug("CameraPlayer14: onPause: stop preview");
+									mCamera.stopPreview();
+								}
+								mCamera.release();
+								mCamera = null;
+								GameActivity.Log.debug("CameraPlayer14: onPause: released camera");
+							}
+							catch (RuntimeException e)
+							{
+								GameActivity.Log.debug("CameraPlayer14: onPause: failed to unlocked camera: " + e);
+							}
+						}
+					}
+					@Override
+					public void onActivityResumed(Activity activity)
+					{
+						if (CameraState == CameraStates.START || CameraState == CameraStates.PAUSE)
+						{
+							try {
+								mCamera = Camera.open(CameraId);
+								GameActivity.Log.debug("CameraPlayer14: onResume: opened camera");
+
+								if (CameraState == CameraStates.START)
+								{
+									setVideoEnabled(mVideoEnabled);
+								}
+							}
+							catch (Exception e)
+							{
+								GameActivity.Log.debug("CameraPlayer14: onResume: failed to reconnect camera: " + e);
+							}
+						}
+					}
+					@Override
+					public void onActivityStopped(Activity activity) 
+					{ 
+					}
+					@Override
+					public void onActivityDestroyed(Activity activity)
+					{
+                    }
+				});
+		}
+	}
+
+	public int getVideoWidth()
+	{
+//		GameActivity.Log.debug("CameraPlayer14: getVideoWidth: " + CameraWidth);
+		return (CameraState == CameraStates.PREPARED ||
+				CameraState == CameraStates.START ||
+				CameraState == CameraStates.PAUSE ||
+				CameraState == CameraStates.STOP) ? CameraWidth : 0;
+	}
+
+	public int getVideoHeight()
+	{
+//		GameActivity.Log.debug("CameraPlayer14: getVideoHeight: " + CameraHeight);
+		return (CameraState == CameraStates.PREPARED ||
+				CameraState == CameraStates.START ||
+				CameraState == CameraStates.PAUSE ||
+				CameraState == CameraStates.STOP) ? CameraHeight : 0;
+	}
+
+	public int getFrameRate()
+	{
+//		GameActivity.Log.debug("CameraPlayer14: getFrameRate: " + CameraFPS);
+		return CameraFPS;
+	}
+
+	public void prepare()
+	{
+		if (null == mCamera)
+		{
+			return;
+		}
+
+//		GameActivity.Log.debug("CameraPlayer14: prepare: " + CameraState);
+//		CameraState = CameraStates.PREPARED;
+		Prepared = true;
+	}
+
+	public void prepareAsync()
+	{
+		if (null == mCamera)
+		{
+			return;
+		}
+
+//		GameActivity.Log.debug("CameraPlayer14: prepareAsync: " + CameraState);
+//		CameraState = CameraStates.PREPARED;
+		Prepared = true;
+	}
+	
+	public boolean isPlaying()
+	{
+		return CameraState == CameraStates.START;
+	}
+
+	public boolean isPrepared()
+	{
+		boolean result;
+		synchronized(this)
+		{
+			result = Prepared;
+		}
+		return result;
+	}
+
+	public boolean didComplete()
+	{
+		boolean result;
+		synchronized(this)
+		{
+			result = Completed;
+			Completed = false;
+		}
+		return result;
+	}
+
+	public boolean isLooping()
+	{
+		return Looping;
+	}
+
+	public int getCurrentPosition()
+	{
+		return 0;
+	}
+
+	public int getDuration()
+	{
+		if (null == mCamera)
+		{
+			return 0;
+		}
+
+		return 1000;
+	}
+
+	public void selectTrack(int track)
+	{
+	}
+
+	public boolean setDataSourceURL(
+		String Url)
+		throws IOException,
+			java.lang.InterruptedException,
+			java.util.concurrent.ExecutionException
+	{
+		if (!CameraSupported)
+		{
+			return false;
+		}
+
+//		GameActivity.Log.debug("setDataSource: Url = " + Url);
+		stop();
+		releaseOESTextureRenderer();
+		releaseBitmapRenderer();
+
+		synchronized(this)
+		{
+			Prepared = false;
+			Completed = false;
+		}
+		Looping = false;
+		audioTracks.clear();
+		videoTracks.clear();
+
+		// defaults		
+		int cameraId = 0;
+		int width = 320;
+		int height = 480;
+		int fps = 30;
+
+		// Url should be in form: "vidcap://camera_id?width=xx?height=xx?fps=xx"
+		int index = Url.indexOf("vidcap://");
+		if (index >= 0)
+		{
+			Url = Url.substring(index + 9);
+		}
+		String[] UrlParts = Url.split("\\?");
+		if (UrlParts.length > 0)
+		{
+//			GameActivity.Log.debug("setDataSource: camera=" + UrlParts[0]);
+			if (UrlParts[0].equals("front"))
+			{
+				cameraId = iFrontId;
+//				GameActivity.Log.debug("setDataSource: front=" + iFrontId);
+			}
+			else if (UrlParts[0].equals("back") || UrlParts[0].equals("rear"))
+			{
+				cameraId = iBackId;
+//				GameActivity.Log.debug("setDataSource: back=" + iBackId);
+			}
+			else
+			{
+				try
+				{
+					cameraId = Integer.parseInt(UrlParts[0]);
+				}
+				catch (NumberFormatException e)
+				{
+				}
+//				GameActivity.Log.debug("setDataSource: back=" + iBackId);
+			}
+
+			for (int UrlPartIndex=1; UrlPartIndex < UrlParts.length; UrlPartIndex++)
+			{
+				String[] KeyValue = UrlParts[UrlPartIndex].split("=", 2);
+				if (KeyValue.length == 2)
+				{
+					int valueInt;
+					try
+					{
+						valueInt = Integer.parseInt(KeyValue[1]);
+					}
+					catch (NumberFormatException e)
+					{
+						continue;
+					}
+					if (KeyValue[0].equals("width"))
+					{
+						width = valueInt;
+					}
+					else if (KeyValue[0].equals("height"))
+					{
+						height = valueInt;
+					}
+					else if (KeyValue[0].equals("fps"))
+					{
+						fps = valueInt;
+					}
+				}
+			}
+		}
+		
+		try {
+			mCamera = Camera.open(cameraId);
+		}
+		catch (Exception e)
+		{
+			mCamera = null;
+			GameActivity.Log.debug("setDataSource: could not open camera: " + cameraId);
+			return false;
+		}
+		
+//		GameActivity.Log.debug("setDataSource: id=" + cameraId + ", width=" + width + ", height=" + height + ", fps=" + fps);
+		
+		Camera.Parameters param = mCamera.getParameters();
+
+		// find min and max fps allowed and matched frame rate range
+		float floatFps = fps;
+		fps *= 1000;
+		int fpsMin = fps;
+		int fpsMax = fps;
+		float floatMinFps = 10000;
+		float floatMaxFps = -10000;
+		List<int[]> previewRates = param.getSupportedPreviewFpsRange();
+		boolean found = false;
+
+		// first look for an exact match
+		for (index = 0; index < previewRates.size(); ++index)
+		{
+			int[] rates = previewRates.get(index);
+//			GameActivity.Log.debug("setDataSource: previewRates " + index + ": " + rates[0] + " - " + rates[1]);
+			float previewMin = (float)rates[0] / 1000.0f;
+			float previewMax = (float)rates[1] / 1000.0f;
+			floatMinFps = previewMin < floatMinFps ? previewMin : floatMinFps;
+			floatMaxFps = previewMax > floatMaxFps ? previewMax : floatMaxFps;
+			if ((fps == rates[0] && fps == rates[1]))
+			{
+				found = true;
+				fpsMin = rates[0];
+				fpsMax = rates[1];
+			}
+		}
+		if (!found)
+		{
+			// look for range containing desired frame rate
+			for (index = 0; index < previewRates.size(); ++index)
+			{
+				int[] rates = previewRates.get(index);
+				if (!found || (fps >= rates[0] && fps <= rates[1]))
+				{
+					found = true;
+					fpsMin = rates[0];
+					fpsMax = rates[1];
+				}
+			}
+		}
+
+		// find best matched preview size and add track info
+		List<Camera.Size> previewSizes = param.getSupportedPreviewSizes();
+		if (previewSizes.size() > 0)
+		{
+			int bestWidth = -1;
+			int bestHeight = -1;
+			int bestScore = 0;
+			for (index = 0; index < previewSizes.size(); ++index)
+			{
+				Camera.Size camSize = previewSizes.get(index);
+//				GameActivity.Log.debug("setDataSource: previewSize " + index + ": " + camSize.width + " x " + camSize.height);
+
+				VideoTrackInfo videoTrack = new VideoTrackInfo();
+				videoTrack.Index = index;
+				videoTrack.MimeType = "video/unknown";
+				videoTrack.DisplayName = "Video Track " + index + " (Stream " + index + ")";
+				videoTrack.Language = "und";
+				videoTrack.BitRate = 0;
+				videoTrack.Width = camSize.width;
+				videoTrack.Height = camSize.height;
+				videoTrack.FrameRate = floatFps;
+				videoTrack.FrameRateLow = floatMinFps;
+				videoTrack.FrameRateHigh = floatMaxFps;
+				videoTracks.add(videoTrack);
+
+				int diffWidth = camSize.width - width;
+				int diffHeight = camSize.height - height;
+				int score = diffWidth * diffWidth + diffHeight * diffHeight;
+				if (bestWidth == -1 || score < bestScore)
+				{
+					bestWidth = camSize.width;
+					bestHeight = camSize.height;
+					bestScore = score;
+				}
+			}
+			width = bestWidth;
+			height = bestHeight;
+		}
+
+		// set camera parameters
+		param.setPreviewSize(width, height);
+		param.setPreviewFpsRange(fpsMin, fpsMax);
+		param.set("orientation", "landscape");
+		mCamera.setDisplayOrientation(0);
+		mCamera.setParameters(param);
+
+		Camera.CameraInfo camInfo = new Camera.CameraInfo();
+		Camera.getCameraInfo(CameraId, camInfo);
+		CameraOrientation = camInfo.orientation;
+
+		switch (GameActivity._activity.getWindowManager().getDefaultDisplay().getRotation())
+		{
+			case Surface.ROTATION_0:	DeviceRotation = 0;		break;
+			case Surface.ROTATION_90:	DeviceRotation = 90;	break;
+			case Surface.ROTATION_180:	DeviceRotation = 180;	break;
+			case Surface.ROTATION_270:	DeviceRotation = 270;	break;
+		}
+
+		// store our picks
+		CameraId = cameraId;
+		CameraWidth = width;
+		CameraHeight = height;
+		CameraFPSMin = fpsMin;
+		CameraFPSMax = fpsMax;
+		CameraFPS = fps / 1000;
+
+		GameActivity.Log.debug("setDataSource: id=" + cameraId + ", width=" + CameraWidth + ", height=" + CameraHeight + ", fps=" + CameraFPS + ", min=" + fpsMin + " max=" + fpsMax + " orientation=" + CameraOrientation);
+
+		CameraState = CameraStates.PREPARED;
+		return true;
+	}
+
+	public void setVideoEnabled(boolean enabled)
+	{
+		if (null == mCamera)
+		{
+			return;
+		}
+
+//		GameActivity.Log.debug("CameraPlayer14: setVideoEnabled: " + enabled);
+		WaitOnBitmapRender = true;
+
+		mVideoEnabled = enabled;
+		if (mVideoEnabled)
+		{
+			if (null != mOESTextureRenderer && null != mOESTextureRenderer.getSurface())
+			{
+				try
+				{
+					mCamera.setPreviewTexture(mOESTextureRenderer.getSurfaceTexture());
+				}
+				catch (IOException e)
+				{
+				}
+			}
+
+			if (null != mBitmapRenderer && null != mBitmapRenderer.getSurface())
+			{
+				try
+				{
+					mCamera.setPreviewTexture(mBitmapRenderer.getSurfaceTexture());
+				}
+				catch (IOException e)
+				{
+				}
+			}
+		}
+		else
+		{
+			try
+			{
+				mCamera.setPreviewTexture(null);
+			}
+			catch (IOException e)
+			{
+			}
+		}
+
+		WaitOnBitmapRender = false;
+		mCamera.startPreview();
+	}
+	
+	public void setAudioEnabled(boolean enabled)
+	{
+	}
+	
+	public boolean didResolutionChange()
+	{
+		if (null != mOESTextureRenderer)
+		{
+			return mOESTextureRenderer.resolutionChanged();
+		}
+		if (null != mBitmapRenderer)
+		{
+			return mBitmapRenderer.resolutionChanged();
+		}
+		return false;
+	}
+
+	public int getExternalTextureId()
+	{
+		if (null != mOESTextureRenderer)
+		{
+			return mOESTextureRenderer.getExternalTextureId();
+		}
+		if (null != mBitmapRenderer)
+		{
+			return mBitmapRenderer.getExternalTextureId();
+		}
+		return -1;
+	}
+
+	public void start()
+	{
+//		GameActivity.Log.debug("CameraPlayer14: start: " + CameraState);
+		if (CameraState == CameraStates.PREPARED)
+		{
+			CameraState = CameraStates.PAUSE;
+		}
+
+		if (CameraState == CameraStates.PAUSE)
+		{
+			mCamera.startPreview();
+			CameraState = CameraStates.START;
+		}
+	}
+
+	public void pause()
+	{
+//		GameActivity.Log.debug("CameraPlayer14: pause: " + CameraState);
+		if (CameraState == CameraStates.START)
+		{
+			mCamera.stopPreview();
+			CameraState = CameraStates.PAUSE;
+		}
+	}
+
+	public void stop()
+	{
+//		GameActivity.Log.debug("CameraPlayer14: stop: " + CameraState);
+		if (CameraState == CameraStates.START)
+		{
+			mCamera.stopPreview();
+			reset();
+		}
+		if (null != mCamera)
+		{
+			mCamera.release();
+			mCamera = null;
+		}
+		CameraState = CameraStates.STOP;
+	}
+
+	public void seekTo(int position)
+	{
+		synchronized (this)
+		{
+			Completed = false;
+		}
+	}
+
+	public void setLooping(boolean looping)
+	{
+		// don't set on player
+		Looping = looping;
+	}
+
+	public void release()
+	{
+//		GameActivity.Log.debug("CameraPlayer14: release");
+		if (null != mOESTextureRenderer)
+		{
+			while (WaitOnBitmapRender) ;
+			releaseOESTextureRenderer();
+		}
+		if (null != mBitmapRenderer)
+		{
+			while (WaitOnBitmapRender) ;
+			releaseOESTextureRenderer();
+		}
+
+		if (null != mCamera)
+		{
+			mCamera.stopPreview();
+			mCamera.release();
+			mCamera = null;
+		}
+	}
+
+	public void reset()
+	{
+//		GameActivity.Log.debug("CameraPlayer14: reset");
+		if (null != mOESTextureRenderer)
+		{
+			while (WaitOnBitmapRender) ;
+			releaseOESTextureRenderer();
+		}
+		if (null != mBitmapRenderer)
+		{
+			while (WaitOnBitmapRender) ;
+			releaseBitmapRenderer();
+		}
+	}
+
+	// ======================================================================================
+
+	private boolean CreateBitmapRenderer()
+	{
+		releaseBitmapRenderer();
+
+		mBitmapRenderer = new BitmapRenderer(SwizzlePixels, VulkanRenderer);
+		if (!mBitmapRenderer.isValid())
+		{
+			mBitmapRenderer = null;
+			return false;
+		}
+		
+		// set this here as the size may have been set before the GL resources were created.
+		mBitmapRenderer.setSize(getVideoWidth(),getVideoHeight());
+
+		setVideoEnabled(true);
+		setAudioEnabled(true);
+		return true;
+	}
+
+	void releaseBitmapRenderer()
+	{
+		if (null != mBitmapRenderer)
+		{
+			mBitmapRenderer.release();
+			mBitmapRenderer = null;
+
+			if (null != mCamera)
+			{
+				try
+				{
+					mCamera.setPreviewTexture(null);
+				}
+				catch (IOException e)
+				{
+				}
+			}
+		}
+	}
+
+	public void initBitmapRenderer()
+	{
+		// if not already allocated.
+		// Create bitmap renderer's gl resources in the renderer thread.
+		if (null == mBitmapRenderer)
+		{
+			if (!CreateBitmapRenderer())
+			{
+				GameActivity.Log.warn("initBitmapRenderer failed to alloc mBitmapRenderer ");
+				reset();
+			  }
+		}
+	}
+
+	public FrameUpdateInfo getVideoLastFrameData()
+	{
+		if (null == mCamera)
+		{
+			return null;
+		}
+
+		initBitmapRenderer();
+		if (null != mBitmapRenderer)
+		{
+			WaitOnBitmapRender = true;
+			FrameUpdateInfo frameInfo = mBitmapRenderer.updateFrameData();
+			WaitOnBitmapRender = false;
+			return frameInfo;
+		}
+		else
+		{
+			return null;
+		}
+	}
+
+	public FrameUpdateInfo getVideoLastFrame(int destTexture)
+	{
+		if (null == mCamera)
+		{
+			return null;
+		}
+
+//		GameActivity.Log.debug("getVideoLastFrame: " + destTexture);
+		initBitmapRenderer();
+		if (null != mBitmapRenderer)
+		{
+			WaitOnBitmapRender = true;
+			FrameUpdateInfo frameInfo = mBitmapRenderer.updateFrameData(destTexture);
+			WaitOnBitmapRender = false;
+			return frameInfo;
+		}
+		else
+		{
+			return null;
+		}
+	}
+
+	/*
+		All this internal surface view does is manage the
+		offscreen bitmap that the media player decoding can
+		render into for eventual extraction to the UE4 buffers.
+	*/
+	class BitmapRenderer
+		implements android.graphics.SurfaceTexture.OnFrameAvailableListener
+	{
+		private java.nio.Buffer mFrameData = null;
+		private int mLastFramePosition = -1;
+		private android.graphics.SurfaceTexture mSurfaceTexture = null;
+		private int mTextureWidth = -1;
+		private int mTextureHeight = -1;
+		private android.view.Surface mSurface = null;
+		private boolean mFrameAvailable = false;
+		private int mTextureID = -1;
+		private int mFBO = -1;
+		private int mBlitVertexShaderID = -1;
+		private int mBlitFragmentShaderID = -1;
+		private float[] mTransformMatrix = new float[16];
+		private boolean mTriangleVerticesDirty = true;
+		private boolean mTextureSizeChanged = true;
+		private boolean mUseOwnContext = true;
+		private boolean mVulkanRenderer = false;
+		private boolean mSwizzlePixels = false;
+
+		private float mScaleRotation00 = 1.0f;
+		private float mScaleRotation01 = 0.0f;
+		private float mScaleRotation10 = 0.0f;
+		private float mScaleRotation11 = 1.0f;
+		private float mUOffset = 0.0f;
+		private float mVOffset = 0.0f;
+
+		private int GL_TEXTURE_EXTERNAL_OES = 0x8D65;
+
+		private EGLDisplay mEglDisplay;
+		private EGLContext mEglContext;
+		private EGLSurface mEglSurface;
+
+		private EGLDisplay mSavedDisplay;
+		private EGLContext mSavedContext;
+		private EGLSurface mSavedSurfaceDraw;
+		private EGLSurface mSavedSurfaceRead;
+
+		private boolean mCreatedEGLDisplay = false;
+
+		public BitmapRenderer(boolean swizzlePixels, boolean isVulkan)
+		{
+			mSwizzlePixels = swizzlePixels;
+			mVulkanRenderer = isVulkan;
+
+			mEglSurface = EGL14.EGL_NO_SURFACE;
+			mEglContext = EGL14.EGL_NO_CONTEXT;
+			mEglDisplay = EGL14.EGL_NO_DISPLAY;
+			mUseOwnContext = true;
+
+			if (mVulkanRenderer)
+			{
+				mSwizzlePixels = true;
+			}
+			else
+			{
+				String RendererString = GLES20.glGetString(GLES20.GL_RENDERER);
+
+				// Do not use shared context if Adreno before 400 or on older Android than Marshmallow
+				if (RendererString.contains("Adreno (TM) "))
+				{
+					int AdrenoVersion = Integer.parseInt(RendererString.substring(12));
+					if (AdrenoVersion < 400 || android.os.Build.VERSION.SDK_INT < 22)
+					{
+						GameActivity.Log.debug("MediaPlayer14: disabled shared GL context on " + RendererString);
+						mUseOwnContext = false;
+					}
+				}
+			}
+
+			if (mUseOwnContext)
+			{
+				initContext();
+				saveContext();
+				makeCurrent();
+				initSurfaceTexture();
+				restoreContext();
+			}
+			else
+			{
+				initSurfaceTexture();
+			}
+		}
+
+		private void initContext()
+		{
+			mEglDisplay = EGL14.EGL_NO_DISPLAY;
+			EGLContext shareContext = EGL14.EGL_NO_CONTEXT;
+
+			if (!mVulkanRenderer)
+			{
+				mEglDisplay = EGL14.eglGetCurrentDisplay();
+				shareContext = EGL14.eglGetCurrentContext();
+			}
+			else
+			{
+				mEglDisplay = EGL14.eglGetDisplay(EGL14.EGL_DEFAULT_DISPLAY);
+				if (mEglDisplay == EGL14.EGL_NO_DISPLAY)
+				{
+					GameActivity.Log.error("unable to get EGL14 display");
+					return;
+				}
+				int[] version = new int[2];
+				if (!EGL14.eglInitialize(mEglDisplay, version, 0, version, 1))
+				{
+					mEglDisplay = null;
+					GameActivity.Log.error("unable to initialize EGL14 display");
+					return;
+				}				
+				
+				mCreatedEGLDisplay = true;
+			}
+
+			int[] configSpec = new int[]
+			{
+				EGL14.EGL_RENDERABLE_TYPE, EGL14.EGL_OPENGL_ES2_BIT,
+				EGL14.EGL_SURFACE_TYPE, EGL14.EGL_PBUFFER_BIT,
+				EGL14.EGL_NONE
+			};
+			EGLConfig[] configs = new EGLConfig[1];
+			int[] num_config = new int[1];
+			EGL14.eglChooseConfig(mEglDisplay, configSpec, 0, configs, 0, 1, num_config, 0);
+			int[] contextAttribs = new int[]
+			{
+				EGL14.EGL_CONTEXT_CLIENT_VERSION, 2,
+				EGL14.EGL_NONE
+			};
+			mEglContext = EGL14.eglCreateContext(mEglDisplay, configs[0], shareContext, contextAttribs, 0);
+
+			if (EGL14.eglQueryString(mEglDisplay, EGL14.EGL_EXTENSIONS).contains("EGL_KHR_surfaceless_context"))
+			{
+				mEglSurface = EGL14.EGL_NO_SURFACE;
+			}
+			else
+			{
+				int[] pbufferAttribs = new int[]
+				{
+					EGL14.EGL_NONE
+				};
+				mEglSurface = EGL14.eglCreatePbufferSurface(mEglDisplay, configs[0], pbufferAttribs, 0);
+			}
+		}
+
+		private void saveContext()
+		{
+			mSavedDisplay = EGL14.eglGetCurrentDisplay();
+			mSavedContext = EGL14.eglGetCurrentContext();
+			mSavedSurfaceDraw = EGL14.eglGetCurrentSurface(EGL14.EGL_DRAW);
+			mSavedSurfaceRead = EGL14.eglGetCurrentSurface(EGL14.EGL_READ);
+		}
+
+		private void makeCurrent()
+		{
+			EGL14.eglMakeCurrent(mEglDisplay, mEglSurface, mEglSurface, mEglContext);
+		}
+
+		private void restoreContext()
+		{
+			EGL14.eglMakeCurrent(mSavedDisplay, mSavedSurfaceDraw, mSavedSurfaceRead, mSavedContext);
+		}
+
+		private void initSurfaceTexture()
+		{
+			int[] textures = new int[1];
+			GLES20.glGenTextures(1, textures, 0);
+			mTextureID = textures[0];
+			if (mTextureID <= 0)
+			{
+				GameActivity.Log.error("mTextureID <= 0");
+				release();
+				return;
+			}
+			mSurfaceTexture = new android.graphics.SurfaceTexture(mTextureID);
+			mSurfaceTexture.setOnFrameAvailableListener(this);
+			mSurface = new android.view.Surface(mSurfaceTexture);
+
+			int[] glInt = new int[1];
+
+			GLES20.glGenFramebuffers(1,glInt,0);
+			mFBO = glInt[0];
+			if (mFBO <= 0)
+			{
+				GameActivity.Log.error("mFBO <= 0");
+				release();
+				return;
+			}
+
+			// Special shaders for blit of movie texture.
+			mBlitVertexShaderID = createShader(GLES20.GL_VERTEX_SHADER, mBlitVextexShader);
+			if (mBlitVertexShaderID == 0)
+			{
+				GameActivity.Log.error("mBlitVertexShaderID == 0");
+				release();
+				return;
+			}
+			int mBlitFragmentShaderID = createShader(GLES20.GL_FRAGMENT_SHADER,
+				mSwizzlePixels ? mBlitFragmentShaderBGRA : mBlitFragmentShaderRGBA);
+			if (mBlitFragmentShaderID == 0)
+			{
+				GameActivity.Log.error("mBlitFragmentShaderID == 0");
+				release();
+				return;
+			}
+			mProgram = GLES20.glCreateProgram();
+			if (mProgram <= 0)
+			{
+				GameActivity.Log.error("mProgram <= 0");
+				release();
+				return;
+			}
+			GLES20.glAttachShader(mProgram, mBlitVertexShaderID);
+			GLES20.glAttachShader(mProgram, mBlitFragmentShaderID);
+			GLES20.glLinkProgram(mProgram);
+			int[] linkStatus = new int[1];
+			GLES20.glGetProgramiv(mProgram, GLES20.GL_LINK_STATUS, linkStatus, 0);
+			if (linkStatus[0] != GLES20.GL_TRUE)
+			{
+				GameActivity.Log.error("Could not link program: ");
+				GameActivity.Log.error(GLES20.glGetProgramInfoLog(mProgram));
+				GLES20.glDeleteProgram(mProgram);
+				mProgram = 0;
+				release();
+				return;
+			}
+			mPositionAttrib = GLES20.glGetAttribLocation(mProgram, "Position");
+			mTexCoordsAttrib = GLES20.glGetAttribLocation(mProgram, "TexCoords");
+			mTextureUniform = GLES20.glGetUniformLocation(mProgram, "VideoTexture");
+
+			GLES20.glGenBuffers(1,glInt,0);
+			mBlitBuffer = glInt[0];
+			if (mBlitBuffer <= 0)
+			{
+				GameActivity.Log.error("mBlitBuffer <= 0");
+				release();
+				return;
+			}
+
+			// Create blit mesh.
+			mTriangleVertices = java.nio.ByteBuffer.allocateDirect(
+				mTriangleVerticesData.length * FLOAT_SIZE_BYTES)
+					.order(java.nio.ByteOrder.nativeOrder()).asFloatBuffer();
+			mTriangleVerticesDirty = true;
+
+			// Set up GL state
+			if (mUseOwnContext)
+			{
+				GLES20.glDisable(GLES20.GL_BLEND);
+				GLES20.glDisable(GLES20.GL_CULL_FACE);
+				GLES20.glDisable(GLES20.GL_SCISSOR_TEST);
+				GLES20.glDisable(GLES20.GL_STENCIL_TEST);
+				GLES20.glDisable(GLES20.GL_DEPTH_TEST);
+				GLES20.glDisable(GLES20.GL_DITHER);
+				GLES20.glColorMask(true,true,true,true);
+			}
+		}
+		
+		private void UpdateVertexData()
+		{
+			if (!mTriangleVerticesDirty || mBlitBuffer <= 0)
+			{
+				return;
+			}
+
+			// fill it in
+			mTriangleVertices.position(0);
+			mTriangleVertices.put(mTriangleVerticesData).position(0);
+
+			// save VBO state
+			int[] glInt = new int[1];
+			GLES20.glGetIntegerv(GLES20.GL_ARRAY_BUFFER_BINDING, glInt, 0);
+			int previousVBO = glInt[0];
+			
+			GLES20.glBindBuffer(GLES20.GL_ARRAY_BUFFER, mBlitBuffer);
+			GLES20.glBufferData(GLES20.GL_ARRAY_BUFFER,
+				mTriangleVerticesData.length*FLOAT_SIZE_BYTES,
+				mTriangleVertices, GLES20.GL_STATIC_DRAW);
+
+			// restore VBO state
+			GLES20.glBindBuffer(GLES20.GL_ARRAY_BUFFER, previousVBO);
+			
+			mTriangleVerticesDirty = false;
+		}
+
+		public boolean isValid()
+		{
+			return mSurfaceTexture != null;
+		}
+
+		private int createShader(int shaderType, String source)
+		{
+			int shader = GLES20.glCreateShader(shaderType);
+			if (shader != 0)
+			{
+				GLES20.glShaderSource(shader, source);
+				GLES20.glCompileShader(shader);
+				int[] compiled = new int[1];
+				GLES20.glGetShaderiv(shader, GLES20.GL_COMPILE_STATUS, compiled, 0);
+				if (compiled[0] == 0)
+				{
+					GameActivity.Log.error("Could not compile shader " + shaderType + ":");
+					GameActivity.Log.error(GLES20.glGetShaderInfoLog(shader));
+					GLES20.glDeleteShader(shader);
+					shader = 0;
+				}
+			}
+			return shader;
+		}
+
+		public void onFrameAvailable(android.graphics.SurfaceTexture st)
+		{
+			synchronized(this)
+			{
+				mFrameAvailable = true;
+			}
+		}
+
+		public android.graphics.SurfaceTexture getSurfaceTexture()
+		{
+			return mSurfaceTexture;
+		}
+
+		public android.view.Surface getSurface()
+		{
+			return mSurface;
+		}
+
+		public int getExternalTextureId()
+		{
+			return mTextureID;
+		}
+
+		// NOTE: Synchronized with updateFrameData to prevent frame
+		// updates while the surface may need to get reallocated.
+		public void setSize(int width, int height)
+		{
+			synchronized(this)
+			{
+				if (width != mTextureWidth ||
+					height != mTextureHeight)
+				{
+					mTextureWidth = width;
+					mTextureHeight = height;
+					mFrameData = null;
+					mTextureSizeChanged = true;
+
+				}
+			}
+		}
+
+		public boolean resolutionChanged()
+		{
+			boolean changed;
+			synchronized(this)
+			{
+				changed = mTextureSizeChanged;
+				mTextureSizeChanged = false;
+			}
+			return changed;
+		}
+
+		private static final int FLOAT_SIZE_BYTES = 4;
+		private static final int TRIANGLE_VERTICES_DATA_STRIDE_BYTES = 4 * FLOAT_SIZE_BYTES;
+		private static final int TRIANGLE_VERTICES_DATA_POS_OFFSET = 0;
+		private static final int TRIANGLE_VERTICES_DATA_UV_OFFSET = 2;
+		private float[] mTriangleVerticesData = {
+			// X, Y, U, V
+			-1.0f, -1.0f, 0.f, 0.f,
+			1.0f, -1.0f, 1.f, 0.f,
+			-1.0f, 1.0f, 0.f, 1.f,
+			1.0f, 1.0f, 1.f, 1.f,
+			};
+
+		private java.nio.FloatBuffer mTriangleVertices;
+
+		private final String mBlitVextexShader =
+			"attribute vec2 Position;\n" +
+			"attribute vec2 TexCoords;\n" +
+			"varying vec2 TexCoord;\n" +
+			"void main() {\n" +
+			"	TexCoord = TexCoords;\n" +
+			"	gl_Position = vec4(Position, 0.0, 1.0);\n" +
+			"}\n";
+
+		// NOTE: We read the fragment as BGRA so that in the end, when
+		// we glReadPixels out of the FBO, we get them in that order
+		// and avoid having to swizzle the pixels in the CPU.
+		private final String mBlitFragmentShaderBGRA =
+			"#extension GL_OES_EGL_image_external : require\n" +
+			"uniform samplerExternalOES VideoTexture;\n" +
+			"varying highp vec2 TexCoord;\n" +
+			"void main()\n" +
+			"{\n" +
+			"	gl_FragColor = texture2D(VideoTexture, TexCoord).bgra;\n" +
+			"}\n";
+		private final String mBlitFragmentShaderRGBA =
+			"#extension GL_OES_EGL_image_external : require\n" +
+			"uniform samplerExternalOES VideoTexture;\n" +
+			"varying highp vec2 TexCoord;\n" +
+			"void main()\n" +
+			"{\n" +
+			"	gl_FragColor = texture2D(VideoTexture, TexCoord).rgba;\n" +
+			"}\n";
+
+		private int mProgram;
+		private int mPositionAttrib;
+		private int mTexCoordsAttrib;
+		private int mBlitBuffer;
+		private int mTextureUniform;
+
+		public FrameUpdateInfo updateFrameData()
+		{
+			synchronized(this)
+			{
+				if (null == mFrameData && mTextureWidth > 0 && mTextureHeight > 0)
+				{
+					mFrameData = java.nio.ByteBuffer.allocateDirect(mTextureWidth*mTextureHeight*4);
+				}
+				// Copy surface texture to frame data.
+				if (!copyFrameTexture(0, mFrameData))
+				{
+					return null;
+				}
+			}
+
+			FrameUpdateInfo frameUpdateInfo = new FrameUpdateInfo();
+
+			frameUpdateInfo.Buffer = mFrameData;
+			frameUpdateInfo.CurrentPosition = getCurrentPosition();
+			frameUpdateInfo.FrameReady = true;
+			frameUpdateInfo.RegionChanged = false;
+
+			frameUpdateInfo.ScaleRotation00 = mScaleRotation00;
+			frameUpdateInfo.ScaleRotation01 = mScaleRotation01;
+			frameUpdateInfo.ScaleRotation10 = mScaleRotation10;
+			frameUpdateInfo.ScaleRotation11 = mScaleRotation11;
+
+			frameUpdateInfo.UOffset = mUOffset;
+			frameUpdateInfo.VOffset = mVOffset;
+
+			return frameUpdateInfo;
+		}
+
+		public FrameUpdateInfo updateFrameData(int destTexture)
+		{
+			synchronized(this)
+			{
+				// Copy surface texture to destination texture.
+				if (!copyFrameTexture(destTexture, null))
+				{
+					return null;
+				}
+			}
+
+			FrameUpdateInfo frameUpdateInfo = new FrameUpdateInfo();
+
+			frameUpdateInfo.Buffer = null;
+			frameUpdateInfo.CurrentPosition = getCurrentPosition();
+			frameUpdateInfo.FrameReady = true;
+			frameUpdateInfo.RegionChanged = false;
+
+			frameUpdateInfo.ScaleRotation00 = mScaleRotation00;
+			frameUpdateInfo.ScaleRotation01 = mScaleRotation01;
+			frameUpdateInfo.ScaleRotation10 = mScaleRotation10;
+			frameUpdateInfo.ScaleRotation11 = mScaleRotation11;
+
+			frameUpdateInfo.UOffset = mUOffset;
+			frameUpdateInfo.VOffset = mVOffset;
+
+			return frameUpdateInfo;
+		}
+
+		// Copy the surface texture to another texture, or to raw data.
+		// Note: copying to raw data creates a temporary FBO texture.
+		private boolean copyFrameTexture(int destTexture, java.nio.Buffer destData)
+		{
+			if (!mFrameAvailable)
+			{
+				// We only return fresh data when we generate it. At other
+				// time we return nothing to indicate that there was nothing
+				// new to return. The media player deals with this by keeping
+				// the last frame around and using that for rendering.
+				return false;
+			}
+			mFrameAvailable = false;
+			int current_frame_position = getCurrentPosition();
+			mLastFramePosition = current_frame_position;
+			if (null == mSurfaceTexture)
+			{
+				// Can't update if there's no surface to update into.
+				return false;
+			}
+
+			int[] glInt = new int[1];
+			boolean[] glBool = new boolean[1];
+
+			// Either use own context or save states
+			boolean previousBlend=false, previousCullFace=false, previousScissorTest=false, previousStencilTest=false, previousDepthTest=false, previousDither=false;
+			int previousFBO=0, previousVBO=0, previousMinFilter=0, previousMagFilter=0;
+			int[] previousViewport = new int[4];
+			if (mUseOwnContext)
+			{
+				// Received reports of these not being preserved when changing contexts
+				GLES20.glActiveTexture(GLES20.GL_TEXTURE0);
+				GLES20.glGetTexParameteriv(GLES20.GL_TEXTURE_2D, GLES20.GL_TEXTURE_MIN_FILTER, glInt, 0);
+				previousMinFilter = glInt[0];
+				GLES20.glGetTexParameteriv(GLES20.GL_TEXTURE_2D, GLES20.GL_TEXTURE_MAG_FILTER, glInt, 0);
+				previousMagFilter = glInt[0];
+
+				saveContext();
+				makeCurrent();
+			}
+			else
+			{
+				// Clear gl errors as they can creep in from the UE4 renderer.
+				GLES20.glGetError();
+
+				previousBlend = GLES20.glIsEnabled(GLES20.GL_BLEND);
+				previousCullFace = GLES20.glIsEnabled(GLES20.GL_CULL_FACE);
+				previousScissorTest = GLES20.glIsEnabled(GLES20.GL_SCISSOR_TEST);
+				previousStencilTest = GLES20.glIsEnabled(GLES20.GL_STENCIL_TEST);
+				previousDepthTest = GLES20.glIsEnabled(GLES20.GL_DEPTH_TEST);
+				previousDither = GLES20.glIsEnabled(GLES20.GL_DITHER);
+				GLES20.glGetIntegerv(GLES20.GL_FRAMEBUFFER_BINDING, glInt, 0);
+				previousFBO = glInt[0];
+				GLES20.glGetIntegerv(GLES20.GL_ARRAY_BUFFER_BINDING, glInt, 0);
+				previousVBO = glInt[0];
+				GLES20.glGetIntegerv(GLES20.GL_VIEWPORT, previousViewport, 0);
+
+				GLES20.glActiveTexture(GLES20.GL_TEXTURE0);
+				GLES20.glGetTexParameteriv(GLES20.GL_TEXTURE_2D, GLES20.GL_TEXTURE_MIN_FILTER, glInt, 0);
+				previousMinFilter = glInt[0];
+				GLES20.glGetTexParameteriv(GLES20.GL_TEXTURE_2D, GLES20.GL_TEXTURE_MAG_FILTER, glInt, 0);
+				previousMagFilter = glInt[0];
+
+				glVerify("save state");
+			}
+
+			// Get the latest video texture frame.
+			mSurfaceTexture.updateTexImage();
+
+			mSurfaceTexture.getTransformMatrix(mTransformMatrix);
+
+			int degrees = DeviceRotation;
+			if (degrees == 0) {
+				mScaleRotation00 = mTransformMatrix[1];
+				mScaleRotation01 = mTransformMatrix[5];
+				mScaleRotation10 = -mTransformMatrix[0];
+				mScaleRotation11 = -mTransformMatrix[4];
+				mUOffset = mTransformMatrix[13];
+				mVOffset = 1.0f - mTransformMatrix[12];
+				if (CameraId == iBackId)
+				{
+					mScaleRotation00 = -mScaleRotation00;
+					mScaleRotation01 = -mScaleRotation01;
+					mUOffset = 1.0f - mUOffset;
+				}
+			}
+			else if (degrees == 90) {
+				mScaleRotation00 = mTransformMatrix[0];
+				mScaleRotation01 = mTransformMatrix[4];
+				mScaleRotation10 = mTransformMatrix[1];
+				mScaleRotation11 = mTransformMatrix[5];
+				mUOffset = mTransformMatrix[12];
+				mVOffset = mTransformMatrix[13];
+				if (CameraId == iFrontId)
+				{
+					mScaleRotation00 = -mScaleRotation00;
+					mScaleRotation01 = -mScaleRotation01;
+					mUOffset = 1.0f - mUOffset;
+				}
+				mScaleRotation10 = -mScaleRotation10;
+				mScaleRotation11 = -mScaleRotation11;
+				mVOffset = 1.0f - mVOffset;
+			}
+			else if (degrees == 180) {
+				mScaleRotation00 = mTransformMatrix[1];
+				mScaleRotation01 = mTransformMatrix[5];
+				mScaleRotation10 = mTransformMatrix[0];
+				mScaleRotation11 = mTransformMatrix[4];
+				mUOffset = mTransformMatrix[13];
+				mVOffset = mTransformMatrix[12];
+				if (CameraId == iFrontId)
+				{
+					mScaleRotation00 = -mScaleRotation00;
+					mScaleRotation01 = -mScaleRotation01;
+					mUOffset = 1.0f - mUOffset;
+				}
+			}
+			else {
+				mScaleRotation00 = mTransformMatrix[0];
+				mScaleRotation01 = mTransformMatrix[4];
+				mScaleRotation10 = mTransformMatrix[1];
+				mScaleRotation11 = mTransformMatrix[5];
+				mUOffset = mTransformMatrix[12];
+				mVOffset = mTransformMatrix[13];
+				if (CameraId == iBackId)
+				{
+					mScaleRotation00 = -mScaleRotation00;
+					mScaleRotation01 = -mScaleRotation01;
+					mUOffset = 1.0f - mUOffset;
+				}
+			}
+
+			if (null != destData)
+			{
+				// Rewind data so that we can write to it.
+				destData.position(0);
+			}
+
+			if (!mUseOwnContext)
+			{
+				GLES20.glDisable(GLES20.GL_BLEND);
+				GLES20.glDisable(GLES20.GL_CULL_FACE);
+				GLES20.glDisable(GLES20.GL_SCISSOR_TEST);
+				GLES20.glDisable(GLES20.GL_STENCIL_TEST);
+				GLES20.glDisable(GLES20.GL_DEPTH_TEST);
+				GLES20.glDisable(GLES20.GL_DITHER);
+				GLES20.glColorMask(true,true,true,true);
+
+				glVerify("reset state");
+			}
+
+			GLES20.glViewport(0, 0, mTextureWidth, mTextureHeight);
+
+			glVerify("set viewport");
+
+			// Set-up FBO target texture..
+			int FBOTextureID = 0;
+			if (null != destData)
+			{
+				// Create temporary FBO for data copy.
+				GLES20.glGenTextures(1,glInt,0);
+				FBOTextureID = glInt[0];
+			}
+			else
+			{
+				// Use the given texture as the FBO.
+				FBOTextureID = destTexture;
+			}
+			// Set the FBO to draw into the texture one-to-one.
+			GLES20.glBindTexture(GLES20.GL_TEXTURE_2D, FBOTextureID);
+			GLES20.glTexParameteri(GLES20.GL_TEXTURE_2D,
+				GLES20.GL_TEXTURE_MIN_FILTER, GLES20.GL_NEAREST);
+			GLES20.glTexParameteri(GLES20.GL_TEXTURE_2D,
+				GLES20.GL_TEXTURE_MAG_FILTER, GLES20.GL_NEAREST);
+			GLES20.glTexParameteri(GLES20.GL_TEXTURE_2D,
+				GLES20.GL_TEXTURE_WRAP_S, GLES20.GL_CLAMP_TO_EDGE);
+			GLES20.glTexParameteri(GLES20.GL_TEXTURE_2D,
+				GLES20.GL_TEXTURE_WRAP_T, GLES20.GL_CLAMP_TO_EDGE);
+			// Create the temp FBO data if needed.
+			if (null != destData)
+			{
+				//int w = 1<<(32-Integer.numberOfLeadingZeros(mTextureWidth-1));
+				//int h = 1<<(32-Integer.numberOfLeadingZeros(mTextureHeight-1));
+				GLES20.glTexImage2D(GLES20.GL_TEXTURE_2D, 0,
+					GLES20.GL_RGBA,
+					mTextureWidth, mTextureHeight,
+					0, GLES20.GL_RGBA, GLES20.GL_UNSIGNED_BYTE, null);
+			}
+
+			glVerify("set-up FBO texture");
+
+			// Set to render to the FBO.
+			GLES20.glBindFramebuffer(GLES20.GL_FRAMEBUFFER, mFBO);
+
+			GLES20.glFramebufferTexture2D(
+				GLES20.GL_FRAMEBUFFER,
+				GLES20.GL_COLOR_ATTACHMENT0,
+				GLES20.GL_TEXTURE_2D, FBOTextureID, 0);
+
+			// check status
+			int status = GLES20.glCheckFramebufferStatus(GLES20.GL_FRAMEBUFFER);
+			if (status != GLES20.GL_FRAMEBUFFER_COMPLETE)
+			{
+				GameActivity.Log.warn("Failed to complete framebuffer attachment ("+status+")");
+			}
+
+			// The special shaders to render from the video texture.
+			GLES20.glUseProgram(mProgram);
+
+			// Set the mesh that renders the video texture.
+			UpdateVertexData();
+			GLES20.glBindBuffer(GLES20.GL_ARRAY_BUFFER, mBlitBuffer);
+			GLES20.glEnableVertexAttribArray(mPositionAttrib);
+			GLES20.glVertexAttribPointer(mPositionAttrib, 2, GLES20.GL_FLOAT, false,
+				TRIANGLE_VERTICES_DATA_STRIDE_BYTES, 0);
+			GLES20.glEnableVertexAttribArray(mTexCoordsAttrib);
+			GLES20.glVertexAttribPointer(mTexCoordsAttrib, 2, GLES20.GL_FLOAT, false,
+				TRIANGLE_VERTICES_DATA_STRIDE_BYTES,
+				TRIANGLE_VERTICES_DATA_UV_OFFSET*FLOAT_SIZE_BYTES);
+
+			glVerify("setup movie texture read");
+
+			GLES20.glClear( GLES20.GL_COLOR_BUFFER_BIT);
+
+			// connect 'VideoTexture' to video source texture (mTextureID).
+			// mTextureID is bound to GL_TEXTURE_EXTERNAL_OES in updateTexImage
+			GLES20.glUniform1i(mTextureUniform, 0);
+			GLES20.glActiveTexture(GLES20.GL_TEXTURE0);
+			GLES20.glBindTexture(GL_TEXTURE_EXTERNAL_OES, mTextureID);
+
+			// Draw the video texture mesh.
+			GLES20.glDrawArrays(GLES20.GL_TRIANGLE_STRIP, 0, 4);
+
+			GLES20.glFlush();
+
+			// Read the FBO texture pixels into raw data.
+			if (null != destData)
+			{
+				GLES20.glReadPixels(
+					0, 0, mTextureWidth, mTextureHeight,
+					GLES20.GL_RGBA, GLES20.GL_UNSIGNED_BYTE,
+					destData);
+			}
+
+			glVerify("draw & read movie texture");
+
+			// Restore state and cleanup.
+			if (mUseOwnContext)
+			{
+				GLES20.glFramebufferTexture2D(
+					GLES20.GL_FRAMEBUFFER,
+					GLES20.GL_COLOR_ATTACHMENT0,
+					GLES20.GL_TEXTURE_2D, 0, 0);
+
+				if (null != destData && FBOTextureID > 0)
+				{
+					glInt[0] = FBOTextureID;
+					GLES20.glDeleteTextures(1, glInt, 0);
+				}
+
+				restoreContext();
+
+				// Restore previous texture filtering
+				GLES20.glTexParameteri(GLES20.GL_TEXTURE_2D, GLES20.GL_TEXTURE_MIN_FILTER, previousMinFilter);
+				GLES20.glTexParameteri(GLES20.GL_TEXTURE_2D, GLES20.GL_TEXTURE_MAG_FILTER, previousMagFilter);
+			}
+			else
+			{
+				GLES20.glBindFramebuffer(GLES20.GL_FRAMEBUFFER, previousFBO);
+				if (null != destData && FBOTextureID > 0)
+				{
+					glInt[0] = FBOTextureID;
+					GLES20.glDeleteTextures(1, glInt, 0);
+				}
+				GLES20.glBindBuffer(GLES20.GL_ARRAY_BUFFER, previousVBO);
+
+				GLES20.glViewport(previousViewport[0], previousViewport[1],	previousViewport[2], previousViewport[3]);
+				if (previousBlend) GLES20.glEnable(GLES20.GL_BLEND);
+				if (previousCullFace) GLES20.glEnable(GLES20.GL_CULL_FACE);
+				if (previousScissorTest) GLES20.glEnable(GLES20.GL_SCISSOR_TEST);
+				if (previousStencilTest) GLES20.glEnable(GLES20.GL_STENCIL_TEST);
+				if (previousDepthTest) GLES20.glEnable(GLES20.GL_DEPTH_TEST);
+				if (previousDither) GLES20.glEnable(GLES20.GL_DITHER);
+
+				GLES20.glTexParameteri(GLES20.GL_TEXTURE_2D, GLES20.GL_TEXTURE_MIN_FILTER, previousMinFilter);
+				GLES20.glTexParameteri(GLES20.GL_TEXTURE_2D, GLES20.GL_TEXTURE_MAG_FILTER, previousMagFilter);
+
+				// invalidate cached state in RHI
+				GLES20.glDisableVertexAttribArray(mPositionAttrib);
+				GLES20.glDisableVertexAttribArray(mTexCoordsAttrib);
+				nativeClearCachedAttributeState(mPositionAttrib, mTexCoordsAttrib);
+			}
+
+			return true;
+		}
+
+		private void showGlError(String op, int error)
+		{
+			switch (error)
+			{
+				case GLES20.GL_INVALID_ENUM:						GameActivity.Log.error("CameraPlayer$BitmapRenderer: " + op + ": glGetError GL_INVALID_ENUM");  break;
+				case GLES20.GL_INVALID_OPERATION:					GameActivity.Log.error("CameraPlayer$BitmapRenderer: " + op + ": glGetError GL_INVALID_OPERATION");  break;
+				case GLES20.GL_INVALID_FRAMEBUFFER_OPERATION:		GameActivity.Log.error("CameraPlayer$BitmapRenderer: " + op + ": glGetError GL_INVALID_FRAMEBUFFER_OPERATION");  break;
+				case GLES20.GL_INVALID_VALUE:						GameActivity.Log.error("CameraPlayer$BitmapRenderer: " + op + ": glGetError GL_INVALID_VALUE");  break;
+				case GLES20.GL_FRAMEBUFFER_INCOMPLETE_ATTACHMENT:	GameActivity.Log.error("CameraPlayer$BitmapRenderer: " + op + ": glGetError GL_FRAMEBUFFER_INCOMPLETE_ATTACHMENT");  break;
+				case GLES20.GL_FRAMEBUFFER_INCOMPLETE_DIMENSIONS:	GameActivity.Log.error("CameraPlayer$BitmapRenderer: " + op + ": glGetError GL_FRAMEBUFFER_INCOMPLETE_DIMENSIONS");  break;
+				case GLES20.GL_FRAMEBUFFER_UNSUPPORTED:				GameActivity.Log.error("CameraPlayer$BitmapRenderer: " + op + ": glGetError GL_FRAMEBUFFER_UNSUPPORTED");  break;
+				case GLES20.GL_OUT_OF_MEMORY:						GameActivity.Log.error("CameraPlayer$BitmapRenderer: " + op + ": glGetError GL_OUT_OF_MEMORY");  break;
+				default:											GameActivity.Log.error("CameraPlayer$BitmapRenderer: " + op + ": glGetError " + error);
+			}
+		}
+
+		private void glVerify(String op)
+		{
+			int error;
+			while ((error = GLES20.glGetError()) != GLES20.GL_NO_ERROR)
+			{
+				showGlError(op, error);
+				throw new RuntimeException(op + ": glGetError " + error);
+			}
+		}
+
+		private void glWarn(String op)
+		{
+			int error;
+			while ((error = GLES20.glGetError()) != GLES20.GL_NO_ERROR)
+			{
+				showGlError(op, error);
+			}
+		}
+
+		public void release()
+		{
+			if (null != mSurface)
+			{
+				mSurface.release();
+				mSurface = null;
+			}
+			if (null != mSurfaceTexture)
+			{
+				mSurfaceTexture.release();
+				mSurfaceTexture = null;
+			}
+			int[] glInt = new int[1];
+			if (mBlitBuffer > 0)
+			{
+				glInt[0] = mBlitBuffer;
+				GLES20.glDeleteBuffers(1,glInt,0);
+				mBlitBuffer = -1;
+			}
+			if (mProgram > 0)
+			{
+				GLES20.glDeleteProgram(mProgram);
+				mProgram = -1;
+			}
+			if (mBlitVertexShaderID > 0)
+			{
+				GLES20.glDeleteShader(mBlitVertexShaderID);
+				mBlitVertexShaderID = -1;
+			}
+			if (mBlitFragmentShaderID > 0)
+			{
+				GLES20.glDeleteShader(mBlitFragmentShaderID);
+				mBlitFragmentShaderID = -1;
+			}
+			if (mFBO > 0)
+			{
+				glInt[0] = mFBO;
+				GLES20.glDeleteFramebuffers(1,glInt,0);
+				mFBO = -1;
+			}
+			if (mTextureID > 0)
+			{
+				glInt[0] = mTextureID;
+				GLES20.glDeleteTextures(1,glInt,0);
+				mTextureID = -1;
+			}
+			if (mEglSurface != EGL14.EGL_NO_SURFACE)
+			{
+				EGL14.eglDestroySurface(mEglDisplay, mEglSurface);
+				mEglSurface = EGL14.EGL_NO_SURFACE;
+			}
+			if (mEglContext != EGL14.EGL_NO_CONTEXT)
+			{
+				EGL14.eglDestroyContext(mEglDisplay, mEglContext);
+				mEglContext = EGL14.EGL_NO_CONTEXT;
+			}
+			if (mCreatedEGLDisplay)
+			{
+				EGL14.eglTerminate(mEglDisplay);
+				mEglDisplay = EGL14.EGL_NO_DISPLAY;
+				mCreatedEGLDisplay = false;	
+			}
+		}
+	};
+
+	public native void nativeClearCachedAttributeState(int PositionAttrib, int TexCoordsAttrib);
+
+	// ======================================================================================
+
+	private boolean CreateOESTextureRenderer(int OESTextureId)
+	{
+		releaseOESTextureRenderer();
+
+		mOESTextureRenderer = new OESTextureRenderer(OESTextureId);
+		if (!mOESTextureRenderer.isValid())
+		{
+			mOESTextureRenderer = null;
+			return false;
+		}
+		
+		// set this here as the size may have been set before the GL resources were created.
+		mOESTextureRenderer.setSize(getVideoWidth(),getVideoHeight());
+
+		setVideoEnabled(true);
+		setAudioEnabled(true);
+		return true;
+	}
+
+	void releaseOESTextureRenderer()
+	{
+		if (null != mOESTextureRenderer)
+		{
+			mOESTextureRenderer.release();
+			mOESTextureRenderer = null;
+
+			if (null != mCamera)
+			{
+				try
+				{
+					mCamera.setPreviewTexture(null);
+				}
+				catch (IOException e)
+				{
+				}
+			}
+		}
+	}
+	
+	public FrameUpdateInfo updateVideoFrame(int externalTextureId)
+	{
+		if (null == mOESTextureRenderer)
+		{
+			if (!CreateOESTextureRenderer(externalTextureId))
+			{
+				GameActivity.Log.warn("updateVideoFrame failed to alloc mOESTextureRenderer ");
+				reset();
+				return null;
+			}
+		}
+
+		WaitOnBitmapRender = true;
+		FrameUpdateInfo result = mOESTextureRenderer.updateVideoFrame();
+		WaitOnBitmapRender = false;
+		return result;
+	}
+
+	/*
+		This handles events for our OES texture
+	*/
+	class OESTextureRenderer
+		implements android.graphics.SurfaceTexture.OnFrameAvailableListener
+	{
+		private android.graphics.SurfaceTexture mSurfaceTexture = null;
+		private int mTextureWidth = -1;
+		private int mTextureHeight = -1;
+		private android.view.Surface mSurface = null;
+		private boolean mFrameAvailable = false;
+		private int mTextureID = -1;
+		private float[] mTransformMatrix = new float[16];
+		private boolean mTextureSizeChanged = true;
+
+		private float mScaleRotation00 = 1.0f;
+		private float mScaleRotation01 = 0.0f;
+		private float mScaleRotation10 = 0.0f;
+		private float mScaleRotation11 = 1.0f;
+		private float mUOffset = 0.0f;
+		private float mVOffset = 0.0f;
+
+		public OESTextureRenderer(int OESTextureId)
+		{
+			mTextureID = OESTextureId;
+
+			mSurfaceTexture = new android.graphics.SurfaceTexture(mTextureID);
+			mSurfaceTexture.setOnFrameAvailableListener(this);
+			mSurface = new android.view.Surface(mSurfaceTexture);
+		}
+
+		public void release()
+		{
+			if (null != mSurface)
+			{
+				mSurface.release();
+				mSurface = null;
+			}
+			if (null != mSurfaceTexture)
+			{
+				mSurfaceTexture.release();
+				mSurfaceTexture = null;
+			}
+		}
+
+		public boolean isValid()
+		{
+			return mSurfaceTexture != null;
+		}
+
+		public void onFrameAvailable(android.graphics.SurfaceTexture st)
+		{
+			synchronized(this)
+			{
+				mFrameAvailable = true;
+			}
+		}
+
+		public android.graphics.SurfaceTexture getSurfaceTexture()
+		{
+			return mSurfaceTexture;
+		}
+
+		public android.view.Surface getSurface()
+		{
+			return mSurface;
+		}
+
+		public int getExternalTextureId()
+		{
+			return mTextureID;
+		}
+
+		// NOTE: Synchronized with updateFrameData to prevent frame
+		// updates while the surface may need to get reallocated.
+		public void setSize(int width, int height)
+		{
+			synchronized(this)
+			{
+				if (width != mTextureWidth ||
+					height != mTextureHeight)
+				{
+					mTextureWidth = width;
+					mTextureHeight = height;
+					mTextureSizeChanged = true;
+				}
+			}
+		}
+
+		public boolean resolutionChanged()
+		{
+			boolean changed;
+			synchronized(this)
+			{
+				changed = mTextureSizeChanged;
+				mTextureSizeChanged = false;
+			}
+			return changed;
+		}
+
+		public FrameUpdateInfo updateVideoFrame()
+		{
+			synchronized(this)
+			{
+				return getFrameUpdateInfo();
+			}
+		}
+
+		private FrameUpdateInfo getFrameUpdateInfo()
+		{
+			FrameUpdateInfo frameUpdateInfo = new FrameUpdateInfo();
+
+			frameUpdateInfo.Buffer = null;
+			frameUpdateInfo.CurrentPosition = getCurrentPosition();
+			frameUpdateInfo.FrameReady = false;
+			frameUpdateInfo.RegionChanged = false;
+
+			frameUpdateInfo.ScaleRotation00 = mScaleRotation00;
+			frameUpdateInfo.ScaleRotation01 = mScaleRotation01;
+			frameUpdateInfo.ScaleRotation10 = mScaleRotation10;
+			frameUpdateInfo.ScaleRotation11 = mScaleRotation11;
+
+			frameUpdateInfo.UOffset = mUOffset;
+			frameUpdateInfo.VOffset = mVOffset;
+
+			if (!mFrameAvailable)
+			{
+				// We only return fresh data when we generate it. At other
+				// time we return nothing to indicate that there was nothing
+				// new to return. The media player deals with this by keeping
+				// the last frame around and using that for rendering.
+				return frameUpdateInfo;
+			}
+			mFrameAvailable = false;
+			if (null == mSurfaceTexture)
+			{
+				// Can't update if there's no surface to update into.
+				return frameUpdateInfo;
+			}
+
+			frameUpdateInfo.FrameReady = true;
+
+			// Get the latest video texture frame.
+			mSurfaceTexture.updateTexImage();
+			mSurfaceTexture.getTransformMatrix(mTransformMatrix);
+
+			/*
+			GameActivity.Log.debug(mTransformMatrix[0] + ", " + mTransformMatrix[1] + ", " + mTransformMatrix[2] + ", " + mTransformMatrix[3]);
+			GameActivity.Log.debug(mTransformMatrix[4] + ", " + mTransformMatrix[5] + ", " + mTransformMatrix[6] + ", " + mTransformMatrix[7]);
+			GameActivity.Log.debug(mTransformMatrix[8] + ", " + mTransformMatrix[9] + ", " + mTransformMatrix[10] + ", " + mTransformMatrix[11]);
+			GameActivity.Log.debug(mTransformMatrix[12] + ", " + mTransformMatrix[13] + ", " + mTransformMatrix[14] + ", " + mTransformMatrix[15]);
+			*/
+
+			int degrees = DeviceRotation;
+			if (degrees == 0) {
+				mScaleRotation00 = mTransformMatrix[1];
+				mScaleRotation01 = mTransformMatrix[5];
+				mScaleRotation10 = -mTransformMatrix[0];
+				mScaleRotation11 = -mTransformMatrix[4];
+				mUOffset = mTransformMatrix[13];
+				mVOffset = 1.0f - mTransformMatrix[12];
+				if (CameraId == iBackId)
+				{
+					mScaleRotation00 = -mScaleRotation00;
+					mScaleRotation01 = -mScaleRotation01;
+					mUOffset = 1.0f - mUOffset;
+				}
+			}
+			else if (degrees == 90) {
+				mScaleRotation00 = mTransformMatrix[0];
+				mScaleRotation01 = mTransformMatrix[4];
+				mScaleRotation10 = mTransformMatrix[1];
+				mScaleRotation11 = mTransformMatrix[5];
+				mUOffset = mTransformMatrix[12];
+				mVOffset = mTransformMatrix[13];
+				if (CameraId == iFrontId)
+				{
+					mScaleRotation00 = -mScaleRotation00;
+					mScaleRotation01 = -mScaleRotation01;
+					mUOffset = 1.0f - mUOffset;
+				}
+				mScaleRotation10 = -mScaleRotation10;
+				mScaleRotation11 = -mScaleRotation11;
+				mVOffset = 1.0f - mVOffset;
+			}
+			else if (degrees == 180) {
+				mScaleRotation00 = mTransformMatrix[1];
+				mScaleRotation01 = mTransformMatrix[5];
+				mScaleRotation10 = mTransformMatrix[0];
+				mScaleRotation11 = mTransformMatrix[4];
+				mUOffset = mTransformMatrix[13];
+				mVOffset = mTransformMatrix[12];
+				if (CameraId == iFrontId)
+				{
+					mScaleRotation00 = -mScaleRotation00;
+					mScaleRotation01 = -mScaleRotation01;
+					mUOffset = 1.0f - mUOffset;
+				}
+			}
+			else {
+				mScaleRotation00 = mTransformMatrix[0];
+				mScaleRotation01 = mTransformMatrix[4];
+				mScaleRotation10 = mTransformMatrix[1];
+				mScaleRotation11 = mTransformMatrix[5];
+				mUOffset = mTransformMatrix[12];
+				mVOffset = mTransformMatrix[13];
+				if (CameraId == iBackId)
+				{
+					mScaleRotation00 = -mScaleRotation00;
+					mScaleRotation01 = -mScaleRotation01;
+					mUOffset = 1.0f - mUOffset;
+				}
+			}
+
+			if (frameUpdateInfo.ScaleRotation00 != mScaleRotation00 ||
+				frameUpdateInfo.ScaleRotation01 != mScaleRotation01 ||
+				frameUpdateInfo.ScaleRotation10 != mScaleRotation10 ||
+				frameUpdateInfo.ScaleRotation11 != mScaleRotation11 ||
+				frameUpdateInfo.UOffset != mUOffset ||
+				frameUpdateInfo.VOffset != mVOffset)
+			{
+				frameUpdateInfo.RegionChanged = true;
+
+				frameUpdateInfo.ScaleRotation00 = mScaleRotation00;
+				frameUpdateInfo.ScaleRotation01 = mScaleRotation01;
+				frameUpdateInfo.ScaleRotation10 = mScaleRotation10;
+				frameUpdateInfo.ScaleRotation11 = mScaleRotation11;
+
+				frameUpdateInfo.UOffset = mUOffset;
+				frameUpdateInfo.VOffset = mVOffset;
+			}
+
+			return frameUpdateInfo;
+		}
+	};
+
+	// ======================================================================================
+
+	public AudioTrackInfo[] GetAudioTracks()
+	{
+		AudioTrackInfo[] AudioTracks = new AudioTrackInfo[0];
+
+		return AudioTracks;
+	}
+
+	public CaptionTrackInfo[] GetCaptionTracks()
+	{
+		CaptionTrackInfo[] CaptionTracks = new CaptionTrackInfo[0];
+
+		return CaptionTracks;
+	}
+
+	public VideoTrackInfo[] GetVideoTracks()
+	{
+		int Width = getVideoWidth();
+		int Height = getVideoHeight();
+
+		int CountTracks = videoTracks.size();
+		if (CountTracks > 0)
+		{
+			VideoTrackInfo[] VideoTracks = new VideoTrackInfo[CountTracks];
+			int TrackIndex = 0;
+			for (VideoTrackInfo track : videoTracks)
+			{
+				VideoTracks[TrackIndex] = new VideoTrackInfo();
+				VideoTracks[TrackIndex].Index = track.Index;
+				VideoTracks[TrackIndex].MimeType = track.MimeType;
+				VideoTracks[TrackIndex].DisplayName = track.DisplayName;
+				VideoTracks[TrackIndex].Language = track.Language;
+				VideoTracks[TrackIndex].BitRate = track.BitRate;
+				VideoTracks[TrackIndex].Width = track.Width;
+				VideoTracks[TrackIndex].Height = track.Height;
+				VideoTracks[TrackIndex].FrameRate = track.FrameRate;
+				VideoTracks[TrackIndex].FrameRateLow = track.FrameRateLow;
+				VideoTracks[TrackIndex].FrameRateHigh = track.FrameRateHigh;
+//				GameActivity.Log.debug(TrackIndex + " (" + Index + ") Video: Mime=" + VideoTracks[TrackIndex].MimeType + ", Lang=" + VideoTracks[TrackIndex].Language + ", Width=" + VideoTracks[TrackIndex].Width + ", Height=" + VideoTracks[TrackIndex].Height + ", FrameRate=" + VideoTracks[TrackIndex].FrameRate);
+				TrackIndex++;
+			}
+
+			return VideoTracks;
+		}
+
+		if (Width > 0 && Height > 0)
+		{
+			VideoTrackInfo[] VideoTracks = new VideoTrackInfo[1];
+
+			VideoTracks[0] = new VideoTrackInfo();
+			VideoTracks[0].Index = 0;
+			VideoTracks[0].DisplayName = "Video Track 0 (Stream 0)";
+			VideoTracks[0].Language = "und";
+			VideoTracks[0].BitRate = 0;
+			VideoTracks[0].MimeType = "video/unknown";
+			VideoTracks[0].Width = Width;
+			VideoTracks[0].Height = Height;
+			VideoTracks[0].FrameRate = CameraFPS;
+
+			return VideoTracks;
+		}
+
+		VideoTrackInfo[] VideoTracks = new VideoTrackInfo[0];
+
+		return VideoTracks;
+	}
+
+	// ======================================================================================
+
+	public boolean takePicture(String Filename, int width, int height)
+	{
+		if (null == mCamera || CameraState != CameraStates.START)
+		{
+			return false;
+		}
+
+		if (width != 0 && height != 0)
+		{
+			// find best matched picture size
+			Camera.Parameters param = mCamera.getParameters();
+			List<Camera.Size> pictureSizes = param.getSupportedPictureSizes();
+			if (pictureSizes.size() > 0)
+			{
+				int bestWidth = -1;
+				int bestHeight = -1;
+				int bestScore = 0;
+				for (int index = 0; index < pictureSizes.size(); ++index)
+				{
+					Camera.Size camSize = pictureSizes.get(index);
+					GameActivity.Log.debug("takePicture: pictureSize " + index + ": " + camSize.width + " x " + camSize.height);
+					int diffWidth = camSize.width - width;
+					int diffHeight = camSize.height - height;
+					int score = diffWidth * diffWidth + diffHeight * diffHeight;
+					if (bestWidth == -1 || score < bestScore)
+					{
+						bestWidth = camSize.width;
+						bestHeight = camSize.height;
+						bestScore = score;
+					}
+				}
+				GameActivity.Log.debug("takePicture: using " + bestWidth + " x " + bestHeight);
+				param.setPictureSize(bestWidth, bestHeight);
+				mCamera.setParameters(param);
+			}
+		}
+		
+		final String OutFilename = Filename;
+		mCamera.takePicture(null, null, new Camera.PictureCallback()
+			{
+				@Override
+				public void onPictureTaken(byte[] data, Camera camera) {
+					FileOutputStream outStream = null;
+					try {
+						outStream = new FileOutputStream(OutFilename);
+						outStream.write(data);
+						outStream.close();
+					}
+					catch (FileNotFoundException e) {
+						e.printStackTrace();
+					}
+					catch (IOException e) {
+						e.printStackTrace();
+					}
+					finally {
+					}
+				}
+			});
+		CameraState = CameraStates.PAUSE;
+		
+		return true;
+	}
+}
diff --git a/Engine/Plugins/Media/AndroidCamera/Source/AndroidCamera/Private/AndroidJavaCameraPlayer.cpp b/Engine/Plugins/Media/AndroidCamera/Source/AndroidCamera/Private/AndroidJavaCameraPlayer.cpp
new file mode 100644
--- /dev/null
+++ b/Engine/Plugins/Media/AndroidCamera/Source/AndroidCamera/Private/AndroidJavaCameraPlayer.cpp
@@ -0,0 +1,626 @@
+// Copyright 1998-2017 Epic Games, Inc. All Rights Reserved.
+
+#include "AndroidJavaCameraPlayer.h"
+#include "AndroidApplication.h"
+
+#if UE_BUILD_SHIPPING
+// always clear any exceptions in SHipping
+#define CHECK_JNI_RESULT(Id) if (Id == 0) { JEnv->ExceptionClear(); }
+#else
+#define CHECK_JNI_RESULT(Id) \
+if (Id == 0) \
+{ \
+	if (bIsOptional) { JEnv->ExceptionClear(); } \
+	else { JEnv->ExceptionDescribe(); checkf(Id != 0, TEXT("Failed to find " #Id)); } \
+}
+#endif
+
+static jfieldID FindField(JNIEnv* JEnv, jclass Class, const ANSICHAR* FieldName, const ANSICHAR* FieldType, bool bIsOptional)
+{
+	jfieldID Field = Class == NULL ? NULL : JEnv->GetFieldID(Class, FieldName, FieldType);
+	CHECK_JNI_RESULT(Field);
+	return Field;
+}
+
+FJavaAndroidCameraPlayer::FJavaAndroidCameraPlayer(bool swizzlePixels, bool vulkanRenderer)
+	: FJavaClassObject(GetClassName(), "(ZZ)V", swizzlePixels, vulkanRenderer)
+	, GetDurationMethod(GetClassMethod("getDuration", "()I"))
+	, ResetMethod(GetClassMethod("reset", "()V"))
+	, GetCurrentPositionMethod(GetClassMethod("getCurrentPosition", "()I"))
+	, DidCompleteMethod(GetClassMethod("didComplete", "()Z"))
+	, IsLoopingMethod(GetClassMethod("isLooping", "()Z"))
+	, IsPlayingMethod(GetClassMethod("isPlaying", "()Z"))
+	, IsPreparedMethod(GetClassMethod("isPrepared", "()Z"))
+	, SetDataSourceURLMethod(GetClassMethod("setDataSourceURL", "(Ljava/lang/String;)Z"))
+	, PrepareMethod(GetClassMethod("prepare", "()V"))
+	, PrepareAsyncMethod(GetClassMethod("prepareAsync", "()V"))
+	, SeekToMethod(GetClassMethod("seekTo", "(I)V"))
+	, SetLoopingMethod(GetClassMethod("setLooping", "(Z)V"))
+	, ReleaseMethod(GetClassMethod("release", "()V"))
+	, GetVideoHeightMethod(GetClassMethod("getVideoHeight", "()I"))
+	, GetVideoWidthMethod(GetClassMethod("getVideoWidth", "()I"))
+	, GetFrameRateMethod(GetClassMethod("getFrameRate", "()I"))
+	, SetVideoEnabledMethod(GetClassMethod("setVideoEnabled", "(Z)V"))
+	, SetAudioEnabledMethod(GetClassMethod("setAudioEnabled", "(Z)V"))
+	, GetVideoLastFrameDataMethod(GetClassMethod("getVideoLastFrameData", "()Lcom/epicgames/ue4/CameraPlayer14$FrameUpdateInfo;"))
+	, StartMethod(GetClassMethod("start", "()V"))
+	, PauseMethod(GetClassMethod("pause", "()V"))
+	, StopMethod(GetClassMethod("stop", "()V"))
+	, GetVideoLastFrameMethod(GetClassMethod("getVideoLastFrame", "(I)Lcom/epicgames/ue4/CameraPlayer14$FrameUpdateInfo;"))
+	, GetAudioTracksMethod(GetClassMethod("GetAudioTracks", "()[Lcom/epicgames/ue4/CameraPlayer14$AudioTrackInfo;"))
+	, GetCaptionTracksMethod(GetClassMethod("GetCaptionTracks", "()[Lcom/epicgames/ue4/CameraPlayer14$CaptionTrackInfo;"))
+	, GetVideoTracksMethod(GetClassMethod("GetVideoTracks", "()[Lcom/epicgames/ue4/CameraPlayer14$VideoTrackInfo;"))
+	, DidResolutionChangeMethod(GetClassMethod("didResolutionChange", "()Z"))
+	, GetExternalTextureIdMethod(GetClassMethod("getExternalTextureId", "()I"))
+	, UpdateVideoFrameMethod(GetClassMethod("updateVideoFrame", "(I)Lcom/epicgames/ue4/CameraPlayer14$FrameUpdateInfo;"))
+	, TakePictureMethod(GetClassMethod("takePicture", "(Ljava/lang/String;II)Z"))
+{
+	VideoTexture = nullptr;
+	bVideoTextureValid = false;
+
+	ScaleRotation.X = 1.0f;
+	ScaleRotation.Y = 0.0f;
+	ScaleRotation.Z = 0.0f;
+	ScaleRotation.W = 1.0f;
+	Offset.X = 0.0f;
+	Offset.Y = 0.0f;
+	Offset.Z = 0.0f;
+	Offset.W = 0.0f;
+
+	bTrackInfoSupported = FAndroidMisc::GetAndroidBuildVersion() >= 16;
+
+	PlayerState = FPlayerState::Inactive;
+
+	if (bTrackInfoSupported)
+	{
+		SelectTrackMethod = GetClassMethod("selectTrack", "(I)V");
+	}
+
+	JNIEnv* JEnv = FAndroidApplication::GetJavaEnv();
+
+	// get field IDs for FrameUpdateInfo class members
+	jclass localFrameUpdateInfoClass = FAndroidApplication::FindJavaClass("com/epicgames/ue4/CameraPlayer14$FrameUpdateInfo");
+	FrameUpdateInfoClass = (jclass)JEnv->NewGlobalRef(localFrameUpdateInfoClass);
+	JEnv->DeleteLocalRef(localFrameUpdateInfoClass);
+	FrameUpdateInfo_Buffer = FindField(JEnv, FrameUpdateInfoClass, "Buffer", "Ljava/nio/Buffer;", false);
+	FrameUpdateInfo_CurrentPosition = FindField(JEnv, FrameUpdateInfoClass, "CurrentPosition", "I", false);
+	FrameUpdateInfo_FrameReady = FindField(JEnv, FrameUpdateInfoClass, "FrameReady", "Z", false);
+	FrameUpdateInfo_RegionChanged = FindField(JEnv, FrameUpdateInfoClass, "RegionChanged", "Z", false);
+	FrameUpdateInfo_ScaleRotation00 = FindField(JEnv, FrameUpdateInfoClass, "ScaleRotation00", "F", false);
+	FrameUpdateInfo_ScaleRotation01 = FindField(JEnv, FrameUpdateInfoClass, "ScaleRotation01", "F", false);
+	FrameUpdateInfo_ScaleRotation10 = FindField(JEnv, FrameUpdateInfoClass, "ScaleRotation10", "F", false);
+	FrameUpdateInfo_ScaleRotation11 = FindField(JEnv, FrameUpdateInfoClass, "ScaleRotation11", "F", false);
+	FrameUpdateInfo_UOffset = FindField(JEnv, FrameUpdateInfoClass, "UOffset", "F", false);
+	FrameUpdateInfo_VOffset = FindField(JEnv, FrameUpdateInfoClass, "VOffset", "F", false);
+
+	// get field IDs for AudioTrackInfo class members
+	jclass localAudioTrackInfoClass = FAndroidApplication::FindJavaClass("com/epicgames/ue4/CameraPlayer14$AudioTrackInfo");
+	AudioTrackInfoClass = (jclass)JEnv->NewGlobalRef(localAudioTrackInfoClass);
+	JEnv->DeleteLocalRef(localAudioTrackInfoClass);
+	AudioTrackInfo_Index = FindField(JEnv, AudioTrackInfoClass, "Index", "I", false);
+	AudioTrackInfo_MimeType = FindField(JEnv, AudioTrackInfoClass, "MimeType", "Ljava/lang/String;", false);
+	AudioTrackInfo_DisplayName = FindField(JEnv, AudioTrackInfoClass, "DisplayName", "Ljava/lang/String;", false);
+	AudioTrackInfo_Language = FindField(JEnv, AudioTrackInfoClass, "Language", "Ljava/lang/String;", false);
+	AudioTrackInfo_Channels = FindField(JEnv, AudioTrackInfoClass, "Channels", "I", false);
+	AudioTrackInfo_SampleRate = FindField(JEnv, AudioTrackInfoClass, "SampleRate", "I", false);
+
+	// get field IDs for CaptionTrackInfo class members
+	jclass localCaptionTrackInfoClass = FAndroidApplication::FindJavaClass("com/epicgames/ue4/CameraPlayer14$CaptionTrackInfo");
+	CaptionTrackInfoClass = (jclass)JEnv->NewGlobalRef(localCaptionTrackInfoClass);
+	JEnv->DeleteLocalRef(localCaptionTrackInfoClass);
+	CaptionTrackInfo_Index = FindField(JEnv, CaptionTrackInfoClass, "Index", "I", false);
+	CaptionTrackInfo_MimeType = FindField(JEnv, CaptionTrackInfoClass, "MimeType", "Ljava/lang/String;", false);
+	CaptionTrackInfo_DisplayName = FindField(JEnv, CaptionTrackInfoClass, "DisplayName", "Ljava/lang/String;", false);
+	CaptionTrackInfo_Language = FindField(JEnv, CaptionTrackInfoClass, "Language", "Ljava/lang/String;", false);
+
+	// get field IDs for VideoTrackInfo class members
+	jclass localVideoTrackInfoClass = FAndroidApplication::FindJavaClass("com/epicgames/ue4/CameraPlayer14$VideoTrackInfo");
+	VideoTrackInfoClass = (jclass)JEnv->NewGlobalRef(localVideoTrackInfoClass);
+	JEnv->DeleteLocalRef(localVideoTrackInfoClass);
+	VideoTrackInfo_Index = FindField(JEnv, VideoTrackInfoClass, "Index", "I", false);
+	VideoTrackInfo_MimeType = FindField(JEnv, VideoTrackInfoClass, "MimeType", "Ljava/lang/String;", false);
+	VideoTrackInfo_DisplayName = FindField(JEnv, VideoTrackInfoClass, "DisplayName", "Ljava/lang/String;", false);
+	VideoTrackInfo_Language = FindField(JEnv, VideoTrackInfoClass, "Language", "Ljava/lang/String;", false);
+	VideoTrackInfo_BitRate = FindField(JEnv, VideoTrackInfoClass, "BitRate", "I", false);
+	VideoTrackInfo_Width = FindField(JEnv, VideoTrackInfoClass, "Width", "I", false);
+	VideoTrackInfo_Height = FindField(JEnv, VideoTrackInfoClass, "Height", "I", false);
+	VideoTrackInfo_FrameRate = FindField(JEnv, VideoTrackInfoClass, "FrameRate", "F", false);
+	VideoTrackInfo_FrameRateLow = FindField(JEnv, VideoTrackInfoClass, "FrameRateLow", "F", false);
+	VideoTrackInfo_FrameRateHigh = FindField(JEnv, VideoTrackInfoClass, "FrameRateHigh", "F", false);
+}
+
+int32 FJavaAndroidCameraPlayer::GetDuration()
+{
+	return CallMethod<int32>(GetDurationMethod);
+}
+
+bool FJavaAndroidCameraPlayer::IsActive()
+{
+	return PlayerState == FPlayerState::Active;
+}
+
+void FJavaAndroidCameraPlayer::Reset()
+{
+	PlayerState = FPlayerState::Inactive;
+
+	ScaleRotation.X = 1.0f;
+	ScaleRotation.Y = 0.0f;
+	ScaleRotation.Z = 0.0f;
+	ScaleRotation.W = 1.0f;
+	Offset.X = 0.0f;
+	Offset.Y = 0.0f;
+
+	CallMethod<void>(ResetMethod);
+}
+
+void FJavaAndroidCameraPlayer::Stop()
+{
+	CallMethod<void>(StopMethod);
+}
+
+int32 FJavaAndroidCameraPlayer::GetCurrentPosition()
+{
+	int32 position = CallMethod<int32>(GetCurrentPositionMethod);
+	return position;
+}
+
+bool FJavaAndroidCameraPlayer::IsLooping()
+{
+	return CallMethod<bool>(IsLoopingMethod);
+}
+
+bool FJavaAndroidCameraPlayer::IsPlaying()
+{
+	return CallMethod<bool>(IsPlayingMethod);
+}
+
+bool FJavaAndroidCameraPlayer::IsPrepared()
+{
+	return CallMethod<bool>(IsPreparedMethod);
+}
+
+bool FJavaAndroidCameraPlayer::DidComplete()
+{
+	return CallMethod<bool>(DidCompleteMethod);
+}
+
+bool FJavaAndroidCameraPlayer::SetDataSource(const FString & Url)
+{
+	ScaleRotation.X = 1.0f;
+	ScaleRotation.Y = 0.0f;
+	ScaleRotation.Z = 0.0f;
+	ScaleRotation.W = 1.0f;
+	Offset.X = 0.0f;
+	Offset.Y = 0.0f;
+
+	bool Result = CallMethod<bool>(SetDataSourceURLMethod, GetJString(Url));
+
+	if (Result)
+	{
+		PlayerState = FPlayerState::Active;
+	}
+	return Result;
+}
+
+bool FJavaAndroidCameraPlayer::Prepare()
+{
+	// This can return an exception in some cases (URL without internet, for example)
+	JNIEnv*	JEnv = FAndroidApplication::GetJavaEnv();
+	JEnv->CallVoidMethod(Object, PrepareMethod.Method);
+	if (JEnv->ExceptionCheck())
+	{
+		JEnv->ExceptionDescribe();
+		JEnv->ExceptionClear();
+		return false;
+	}
+	return true;
+}
+
+bool FJavaAndroidCameraPlayer::PrepareAsync()
+{
+	// This can return an exception in some cases (URL without internet, for example)
+	JNIEnv*	JEnv = FAndroidApplication::GetJavaEnv();
+	JEnv->CallVoidMethod(Object, PrepareAsyncMethod.Method);
+	if (JEnv->ExceptionCheck())
+	{
+		JEnv->ExceptionDescribe();
+		JEnv->ExceptionClear();
+		return false;
+	}
+	return true;
+}
+
+void FJavaAndroidCameraPlayer::SeekTo(int32 Milliseconds)
+{
+	CallMethod<void>(SeekToMethod, Milliseconds);
+}
+
+void FJavaAndroidCameraPlayer::SetLooping(bool Looping)
+{
+	CallMethod<void>(SetLoopingMethod, Looping);
+}
+
+void FJavaAndroidCameraPlayer::Release()
+{
+	CallMethod<void>(ReleaseMethod);
+}
+
+int32 FJavaAndroidCameraPlayer::GetVideoHeight()
+{
+	return CallMethod<int32>(GetVideoHeightMethod);
+}
+
+int32 FJavaAndroidCameraPlayer::GetVideoWidth()
+{
+	return CallMethod<int32>(GetVideoWidthMethod);
+}
+
+float FJavaAndroidCameraPlayer::GetFrameRate()
+{
+	return (float)CallMethod<int32>(GetFrameRateMethod);
+}
+
+void FJavaAndroidCameraPlayer::SetVideoEnabled(bool enabled /*= true*/)
+{
+	CallMethod<void>(SetVideoEnabledMethod, enabled);
+}
+
+void FJavaAndroidCameraPlayer::SetAudioEnabled(bool enabled /*= true*/)
+{
+	CallMethod<void>(SetAudioEnabledMethod, enabled);
+}
+
+bool FJavaAndroidCameraPlayer::GetVideoLastFrameData(void* & outPixels, int64 & outCount, int32 *CurrentPosition, bool *bRegionChanged)
+{
+	// This can return an exception in some cases
+	JNIEnv*	JEnv = FAndroidApplication::GetJavaEnv();
+	jobject Result = JEnv->CallObjectMethod(Object, GetVideoLastFrameDataMethod.Method);
+	if (JEnv->ExceptionCheck())
+	{
+		JEnv->ExceptionDescribe();
+		JEnv->ExceptionClear();
+		if (nullptr != Result)
+		{
+			JEnv->DeleteLocalRef(Result);
+		}
+		*CurrentPosition = -1;
+		*bRegionChanged = false;
+		return false;
+	}
+
+	if (nullptr == Result)
+	{
+		return false;
+	}
+
+	jobject buffer = JEnv->GetObjectField(Result, FrameUpdateInfo_Buffer);
+	if (nullptr != buffer)
+	{
+		*CurrentPosition = (int32)JEnv->GetIntField(Result, FrameUpdateInfo_CurrentPosition);
+		bool bFrameReady = (bool)JEnv->GetBooleanField(Result, FrameUpdateInfo_FrameReady);
+		*bRegionChanged = (bool)JEnv->GetBooleanField(Result, FrameUpdateInfo_RegionChanged);
+		ScaleRotation.X = (float)(float)JEnv->GetFloatField(Result, FrameUpdateInfo_ScaleRotation00);
+		ScaleRotation.Y = (float)(float)JEnv->GetFloatField(Result, FrameUpdateInfo_ScaleRotation01);
+		ScaleRotation.Z = (float)(float)JEnv->GetFloatField(Result, FrameUpdateInfo_ScaleRotation10);
+		ScaleRotation.W = (float)(float)JEnv->GetFloatField(Result, FrameUpdateInfo_ScaleRotation11);
+		Offset.X = (float)JEnv->GetFloatField(Result, FrameUpdateInfo_UOffset);
+		Offset.Y = (float)JEnv->GetFloatField(Result, FrameUpdateInfo_VOffset);
+
+		outPixels = JEnv->GetDirectBufferAddress(buffer);
+		outCount = JEnv->GetDirectBufferCapacity(buffer);
+
+		// the GetObjectField returns a local ref, but Java will still own the real buffer
+		JEnv->DeleteLocalRef(buffer);
+
+		JEnv->DeleteLocalRef(Result);
+
+		return !(nullptr == outPixels || 0 == outCount);
+	}
+
+	JEnv->DeleteLocalRef(Result);
+	return false;
+}
+
+void FJavaAndroidCameraPlayer::Start()
+{
+	CallMethod<void>(StartMethod);
+}
+
+void FJavaAndroidCameraPlayer::Pause()
+{
+	CallMethod<void>(PauseMethod);
+}
+
+bool FJavaAndroidCameraPlayer::DidResolutionChange()
+{
+	return CallMethod<bool>(DidResolutionChangeMethod);
+}
+
+int32 FJavaAndroidCameraPlayer::GetExternalTextureId()
+{
+	return CallMethod<int32>(GetExternalTextureIdMethod);
+}
+
+bool FJavaAndroidCameraPlayer::UpdateVideoFrame(int32 ExternalTextureId, int32 *CurrentPosition, bool *bRegionChanged)
+{
+	// This can return an exception in some cases
+	JNIEnv*	JEnv = FAndroidApplication::GetJavaEnv();
+	jobject Result = JEnv->CallObjectMethod(Object, UpdateVideoFrameMethod.Method, ExternalTextureId);
+	if (JEnv->ExceptionCheck())
+	{
+		JEnv->ExceptionDescribe();
+		JEnv->ExceptionClear();
+		if (nullptr != Result)
+		{
+			JEnv->DeleteLocalRef(Result);
+		}
+		*CurrentPosition = -1;
+		*bRegionChanged = false;
+		return false;
+	}
+
+	if (nullptr == Result)
+	{
+		*CurrentPosition = -1;
+		*bRegionChanged = false;
+		return false;
+	}
+
+	*CurrentPosition = (int32)JEnv->GetIntField(Result, FrameUpdateInfo_CurrentPosition);
+	bool bFrameReady = (bool)JEnv->GetBooleanField(Result, FrameUpdateInfo_FrameReady);
+	*bRegionChanged = (bool)JEnv->GetBooleanField(Result, FrameUpdateInfo_RegionChanged);
+	ScaleRotation.X = (float)(float)JEnv->GetFloatField(Result, FrameUpdateInfo_ScaleRotation00);
+	ScaleRotation.Y = (float)(float)JEnv->GetFloatField(Result, FrameUpdateInfo_ScaleRotation01);
+	ScaleRotation.Z = (float)(float)JEnv->GetFloatField(Result, FrameUpdateInfo_ScaleRotation10);
+	ScaleRotation.W = (float)(float)JEnv->GetFloatField(Result, FrameUpdateInfo_ScaleRotation11);
+	Offset.X = (float)JEnv->GetFloatField(Result, FrameUpdateInfo_UOffset);
+	Offset.Y = (float)JEnv->GetFloatField(Result, FrameUpdateInfo_VOffset);
+
+	JEnv->DeleteLocalRef(Result);
+
+	return bFrameReady;
+}
+
+bool FJavaAndroidCameraPlayer::GetVideoLastFrame(int32 destTexture)
+{
+	// This can return an exception in some cases
+	JNIEnv*	JEnv = FAndroidApplication::GetJavaEnv();
+	jobject Result = JEnv->CallObjectMethod(Object, GetVideoLastFrameMethod.Method, destTexture);
+	if (JEnv->ExceptionCheck())
+	{
+		JEnv->ExceptionDescribe();
+		JEnv->ExceptionClear();
+		return false;
+	}
+
+	if (nullptr == Result)
+	{
+		return false;
+	}
+
+	bool bFrameReady = (bool)JEnv->GetBooleanField(Result, FrameUpdateInfo_FrameReady);
+	ScaleRotation.X = (float)(float)JEnv->GetFloatField(Result, FrameUpdateInfo_ScaleRotation00);
+	ScaleRotation.Y = (float)(float)JEnv->GetFloatField(Result, FrameUpdateInfo_ScaleRotation01);
+	ScaleRotation.Z = (float)(float)JEnv->GetFloatField(Result, FrameUpdateInfo_ScaleRotation10);
+	ScaleRotation.W = (float)(float)JEnv->GetFloatField(Result, FrameUpdateInfo_ScaleRotation11);
+	Offset.X = (float)JEnv->GetFloatField(Result, FrameUpdateInfo_UOffset);
+	Offset.Y = (float)JEnv->GetFloatField(Result, FrameUpdateInfo_VOffset);
+
+	JEnv->DeleteLocalRef(Result);
+
+	return bFrameReady;
+}
+
+bool FJavaAndroidCameraPlayer::TakePicture(const FString& Filename)
+{
+	return TakePicture(Filename, 0, 0);
+}
+
+bool FJavaAndroidCameraPlayer::TakePicture(const FString& Filename, int32 Width, int32 Height)
+{
+	return CallMethod<bool>(TakePictureMethod, GetJString(Filename), Width, Height);
+}
+
+FName FJavaAndroidCameraPlayer::GetClassName()
+{
+	if (FAndroidMisc::GetAndroidBuildVersion() >= 14)
+	{
+		return FName("com/epicgames/ue4/CameraPlayer14");
+	}
+	else
+	{
+		return FName("");
+	}
+}
+
+bool FJavaAndroidCameraPlayer::SelectTrack(int32 index)
+{
+	if (!bTrackInfoSupported)
+	{
+		// Just assume it worked
+		return true;
+	}
+
+	// This can return an exception in some cases
+	JNIEnv*	JEnv = FAndroidApplication::GetJavaEnv();
+	JEnv->CallVoidMethod(Object, SelectTrackMethod.Method, index);
+	if (JEnv->ExceptionCheck())
+	{
+		JEnv->ExceptionDescribe();
+		JEnv->ExceptionClear();
+		return false;
+	}
+	return true;
+}
+
+bool FJavaAndroidCameraPlayer::GetAudioTracks(TArray<FAudioTrack>& AudioTracks)
+{
+	AudioTracks.Empty();
+
+	jobjectArray TrackArray = CallMethod<jobjectArray>(GetAudioTracksMethod);
+	if (nullptr != TrackArray)
+	{
+		bool bIsOptional = false;
+		JNIEnv* JEnv = FAndroidApplication::GetJavaEnv();
+		jsize ElementCount = JEnv->GetArrayLength(TrackArray);
+
+		for (int Index = 0; Index < ElementCount; ++Index)
+		{
+			jobject Track = JEnv->GetObjectArrayElement(TrackArray, Index);
+
+			int32 AudioTrackIndex = AudioTracks.AddDefaulted();
+			FAudioTrack& AudioTrack = AudioTracks[AudioTrackIndex];
+
+			AudioTrack.Index = (int32)JEnv->GetIntField(Track, AudioTrackInfo_Index);
+
+			jstring jsMimeType = (jstring)JEnv->GetObjectField(Track, AudioTrackInfo_MimeType);
+			CHECK_JNI_RESULT(jsMimeType);
+			const char * nativeMimeType = JEnv->GetStringUTFChars(jsMimeType, 0);
+			AudioTrack.MimeType = FString(nativeMimeType);
+			JEnv->ReleaseStringUTFChars(jsMimeType, nativeMimeType);
+			JEnv->DeleteLocalRef(jsMimeType);
+
+			jstring jsDisplayName = (jstring)JEnv->GetObjectField(Track, AudioTrackInfo_DisplayName);
+			CHECK_JNI_RESULT(jsDisplayName);
+			const char * nativeDisplayName = JEnv->GetStringUTFChars(jsDisplayName, 0);
+			AudioTrack.DisplayName = FString(nativeDisplayName);
+			JEnv->ReleaseStringUTFChars(jsDisplayName, nativeDisplayName);
+			JEnv->DeleteLocalRef(jsDisplayName);
+
+			jstring jsLanguage = (jstring)JEnv->GetObjectField(Track, AudioTrackInfo_Language);
+			CHECK_JNI_RESULT(jsLanguage);
+			const char * nativeLanguage = JEnv->GetStringUTFChars(jsLanguage, 0);
+			AudioTrack.Language = FString(nativeLanguage);
+			JEnv->ReleaseStringUTFChars(jsLanguage, nativeLanguage);
+			JEnv->DeleteLocalRef(jsLanguage);
+
+			AudioTrack.Channels = (int32)JEnv->GetIntField(Track, AudioTrackInfo_Channels);
+			AudioTrack.SampleRate = (int32)JEnv->GetIntField(Track, AudioTrackInfo_SampleRate);
+		}
+		JEnv->DeleteGlobalRef(TrackArray);
+
+		return true;
+	}
+	return false;
+}
+
+bool FJavaAndroidCameraPlayer::GetCaptionTracks(TArray<FCaptionTrack>& CaptionTracks)
+{
+	CaptionTracks.Empty();
+
+	jobjectArray TrackArray = CallMethod<jobjectArray>(GetCaptionTracksMethod);
+	if (nullptr != TrackArray)
+	{
+		bool bIsOptional = false;
+		JNIEnv* JEnv = FAndroidApplication::GetJavaEnv();
+		jsize ElementCount = JEnv->GetArrayLength(TrackArray);
+
+		for (int Index = 0; Index < ElementCount; ++Index)
+		{
+			jobject Track = JEnv->GetObjectArrayElement(TrackArray, Index);
+
+			int32 CaptionTrackIndex = CaptionTracks.AddDefaulted();
+			FCaptionTrack& CaptionTrack = CaptionTracks[CaptionTrackIndex];
+
+			CaptionTrack.Index = (int32)JEnv->GetIntField(Track, CaptionTrackInfo_Index);
+
+			jstring jsMimeType = (jstring)JEnv->GetObjectField(Track, CaptionTrackInfo_MimeType);
+			CHECK_JNI_RESULT(jsMimeType);
+			const char * nativeMimeType = JEnv->GetStringUTFChars(jsMimeType, 0);
+			CaptionTrack.MimeType = FString(nativeMimeType);
+			JEnv->ReleaseStringUTFChars(jsMimeType, nativeMimeType);
+			JEnv->DeleteLocalRef(jsMimeType);
+
+			jstring jsDisplayName = (jstring)JEnv->GetObjectField(Track, CaptionTrackInfo_DisplayName);
+			CHECK_JNI_RESULT(jsDisplayName);
+			const char * nativeDisplayName = JEnv->GetStringUTFChars(jsDisplayName, 0);
+			CaptionTrack.DisplayName = FString(nativeDisplayName);
+			JEnv->ReleaseStringUTFChars(jsDisplayName, nativeDisplayName);
+			JEnv->DeleteLocalRef(jsDisplayName);
+
+			jstring jsLanguage = (jstring)JEnv->GetObjectField(Track, CaptionTrackInfo_Language);
+			CHECK_JNI_RESULT(jsLanguage);
+			const char * nativeLanguage = JEnv->GetStringUTFChars(jsLanguage, 0);
+			CaptionTrack.Language = FString(nativeLanguage);
+			JEnv->ReleaseStringUTFChars(jsLanguage, nativeLanguage);
+			JEnv->DeleteLocalRef(jsLanguage);
+		}
+		JEnv->DeleteGlobalRef(TrackArray);
+
+		return true;
+	}
+	return false;
+}
+
+bool FJavaAndroidCameraPlayer::GetVideoTracks(TArray<FVideoTrack>& VideoTracks)
+{
+	VideoTracks.Empty();
+
+	jobjectArray TrackArray = CallMethod<jobjectArray>(GetVideoTracksMethod);
+	if (nullptr != TrackArray)
+	{
+		bool bIsOptional = false;
+		JNIEnv* JEnv = FAndroidApplication::GetJavaEnv();
+		jsize ElementCount = JEnv->GetArrayLength(TrackArray);
+
+		if (ElementCount > 0)
+		{
+			jobject Track = JEnv->GetObjectArrayElement(TrackArray, 0);
+
+			int32 VideoTrackIndex = VideoTracks.AddDefaulted();
+			FVideoTrack& VideoTrack = VideoTracks[VideoTrackIndex];
+
+			VideoTrack.Index = (int32)JEnv->GetIntField(Track, VideoTrackInfo_Index);
+
+			jstring jsMimeType = (jstring)JEnv->GetObjectField(Track, VideoTrackInfo_MimeType);
+			CHECK_JNI_RESULT(jsMimeType);
+			const char * nativeMimeType = JEnv->GetStringUTFChars(jsMimeType, 0);
+			VideoTrack.MimeType = FString(nativeMimeType);
+			JEnv->ReleaseStringUTFChars(jsMimeType, nativeMimeType);
+			JEnv->DeleteLocalRef(jsMimeType);
+
+			jstring jsDisplayName = (jstring)JEnv->GetObjectField(Track, VideoTrackInfo_DisplayName);
+			CHECK_JNI_RESULT(jsDisplayName);
+			const char * nativeDisplayName = JEnv->GetStringUTFChars(jsDisplayName, 0);
+			VideoTrack.DisplayName = FString(nativeDisplayName);
+			JEnv->ReleaseStringUTFChars(jsDisplayName, nativeDisplayName);
+			JEnv->DeleteLocalRef(jsDisplayName);
+
+			jstring jsLanguage = (jstring)JEnv->GetObjectField(Track, VideoTrackInfo_Language);
+			CHECK_JNI_RESULT(jsLanguage);
+			const char * nativeLanguage = JEnv->GetStringUTFChars(jsLanguage, 0);
+			VideoTrack.Language = FString(nativeLanguage);
+			JEnv->ReleaseStringUTFChars(jsLanguage, nativeLanguage);
+			JEnv->DeleteLocalRef(jsLanguage);
+
+			VideoTrack.BitRate = (int32)JEnv->GetIntField(Track, VideoTrackInfo_BitRate);
+			VideoTrack.Dimensions = FIntPoint(GetVideoWidth(), GetVideoHeight());
+			VideoTrack.FrameRate = GetFrameRate();
+			VideoTrack.FrameRates = TRange<float>(JEnv->GetFloatField(Track, VideoTrackInfo_FrameRateLow), JEnv->GetFloatField(Track, VideoTrackInfo_FrameRateHigh));
+			VideoTrack.Format = 0;
+
+			for (int Index = 0; Index < ElementCount; ++Index)
+			{
+				jobject Format = JEnv->GetObjectArrayElement(TrackArray, Index);
+
+				int32 VideoFormatIndex = VideoTrack.Formats.AddDefaulted();
+				FVideoFormat& VideoFormat = VideoTrack.Formats[VideoFormatIndex];
+
+				VideoFormat.Dimensions = FIntPoint((int32)JEnv->GetIntField(Format, VideoTrackInfo_Width), (int32)JEnv->GetIntField(Format, VideoTrackInfo_Height));
+				VideoFormat.FrameRate = JEnv->GetFloatField(Format, VideoTrackInfo_FrameRateHigh);
+				VideoFormat.FrameRates = TRange<float>(JEnv->GetFloatField(Format, VideoTrackInfo_FrameRateLow), JEnv->GetFloatField(Format, VideoTrackInfo_FrameRateHigh));
+
+				if (VideoTrack.Dimensions == VideoFormat.Dimensions)
+				{
+					VideoTrack.Format = VideoFormatIndex;
+					VideoFormat.FrameRate = VideoTrack.FrameRate;
+				}
+			}
+		}
+		JEnv->DeleteGlobalRef(TrackArray);
+
+		return true;
+	}
+	return false;
+}
+
diff --git a/Engine/Plugins/Media/AndroidCamera/Source/AndroidCameraEditor/AndroidCameraEditor.Build.cs b/Engine/Plugins/Media/AndroidCamera/Source/AndroidCameraEditor/AndroidCameraEditor.Build.cs
new file mode 100644
--- /dev/null
+++ b/Engine/Plugins/Media/AndroidCamera/Source/AndroidCameraEditor/AndroidCameraEditor.Build.cs
@@ -0,0 +1,23 @@
+// Copyright 1998-2017 Epic Games, Inc. All Rights Reserved.
+
+namespace UnrealBuildTool.Rules
+{
+	public class AndroidCameraEditor : ModuleRules
+	{
+		public AndroidCameraEditor(ReadOnlyTargetRules Target) : base(Target)
+		{
+			PrivateDependencyModuleNames.AddRange(
+				new string[] {
+					"Core",
+					"CoreUObject",
+					"MediaAssets",
+					"UnrealEd",
+				});
+
+			PrivateIncludePaths.AddRange(
+				new string[] {
+					"AndroidCameraEditor/Private",
+				});
+		}
+	}
+}
diff --git a/Engine/Plugins/Media/AndroidMedia/AndroidMedia.uplugin b/Engine/Plugins/Media/AndroidMedia/AndroidMedia.uplugin
--- a/Engine/Plugins/Media/AndroidMedia/AndroidMedia.uplugin
+++ b/Engine/Plugins/Media/AndroidMedia/AndroidMedia.uplugin
@@ -27,11 +27,16 @@
 			"Type" : "Editor",
 			"LoadingPhase" : "PostEngineInit"
 		},
+		{
+			"Name" : "AndroidMediaFactory",
+			"Type" : "Editor",
+			"LoadingPhase" : "PostEngineInit"
+		},
 		{
 			"Name" : "AndroidMediaFactory",
 			"Type" : "RuntimeNoCommandlet",
 			"LoadingPhase" : "PostEngineInit",
-			"WhitelistPlatforms" : [ "Android", "Mac", "Linux", "Win32", "Win64" ]
+			"WhitelistPlatforms" : [ "Android" ]
 		}
 	]
 }
\ No newline at end of file
diff --git a/Engine/Plugins/Media/AndroidMedia/Source/AndroidMedia/AndroidMedia.Build.cs b/Engine/Plugins/Media/AndroidMedia/Source/AndroidMedia/AndroidMedia.Build.cs
--- a/Engine/Plugins/Media/AndroidMedia/Source/AndroidMedia/AndroidMedia.Build.cs
+++ b/Engine/Plugins/Media/AndroidMedia/Source/AndroidMedia/AndroidMedia.Build.cs
@@ -4,29 +4,32 @@ namespace UnrealBuildTool.Rules
 {
 	public class AndroidMedia : ModuleRules
 	{
-        public AndroidMedia(ReadOnlyTargetRules Target) : base(Target)
+		public AndroidMedia(ReadOnlyTargetRules Target) : base(Target)
 		{
-            DynamicallyLoadedModuleNames.AddRange(
-                new string[] {
-                    "Media",
+			DynamicallyLoadedModuleNames.AddRange(
+				new string[] {
+					"Media",
 				});
 
 			PrivateDependencyModuleNames.AddRange(
 				new string[] {
+					"AndroidMediaFactory",
 					"Core",
-				    "Engine",
-                    "RenderCore",
+					"Engine",
+					"MediaUtils",
+					"RenderCore",
 				});
 
 			PrivateIncludePathModuleNames.AddRange(
 				new string[] {
-                    "Media",
+					"Media",
 				});
 
 			PrivateIncludePaths.AddRange(
 				new string[] {
 					"AndroidMedia/Private",
 					"AndroidMedia/Private/Player",
+					"AndroidMedia/Private/Shared",
 				});
 		}
 	}
diff --git a/Engine/Plugins/Media/AndroidMedia/Source/AndroidMedia/Private/AndroidMediaModule.cpp b/Engine/Plugins/Media/AndroidMedia/Source/AndroidMedia/Private/AndroidMediaModule.cpp
--- a/Engine/Plugins/Media/AndroidMedia/Source/AndroidMedia/Private/AndroidMediaModule.cpp
+++ b/Engine/Plugins/Media/AndroidMedia/Source/AndroidMedia/Private/AndroidMediaModule.cpp
@@ -1,15 +1,19 @@
 // Copyright 1998-2017 Epic Games, Inc. All Rights Reserved.
 
-#include "CoreMinimal.h"
-#include "Player/AndroidMediaPlayer.h"
-#include "IAndroidMediaModule.h"
+#include "AndroidMediaPrivate.h"
+
 #include "Modules/ModuleManager.h"
 
+#include "IAndroidMediaModule.h"
+#include "Player/AndroidMediaPlayer.h"
 
 
-#define LOCTEXT_NAMESPACE "FAndroidMediaModule"
+DEFINE_LOG_CATEGORY(LogAndroidMedia);
 
 
+/**
+ * Implements the AndroidMedia module.
+ */
 class FAndroidMediaModule
 	: public IAndroidMediaModule
 {
@@ -17,14 +21,14 @@ public:
 
 	//~ IAndroidMediaModule interface
 
-	virtual TSharedPtr<IMediaPlayer, ESPMode::ThreadSafe> CreatePlayer() override
+	virtual TSharedPtr<IMediaPlayer, ESPMode::ThreadSafe> CreatePlayer(IMediaEventSink& EventSink) override
 	{
 		if (!IsSupported())
 		{
 			return nullptr;
 		}
 
-		return MakeShareable(new FAndroidMediaPlayer());
+		return MakeShared<FAndroidMediaPlayer, ESPMode::ThreadSafe>(EventSink);
 	}
 
 public:
@@ -49,5 +53,3 @@ protected:
 
 
 IMPLEMENT_MODULE(FAndroidMediaModule, AndroidMedia)
-
-#undef LOCTEXT_NAMESPACE
