Fix external texture issues on Android

diff --git a/Engine/Build/Android/Java/src/com/epicgames/ue4/MediaPlayer14.java b/Engine/Build/Android/Java/src/com/epicgames/ue4/MediaPlayer14.java
--- a/Engine/Build/Android/Java/src/com/epicgames/ue4/MediaPlayer14.java
+++ b/Engine/Build/Android/Java/src/com/epicgames/ue4/MediaPlayer14.java
@@ -1694,6 +1694,7 @@ public class MediaPlayer14
 		private int mTextureID = -1;
 		private float[] mTransformMatrix = new float[16];
 		private boolean mTextureSizeChanged = true;
+		private int GL_TEXTURE_EXTERNAL_OES = 0x8D65;
 
 		private float mUScale = 1.0f;
 		private float mVScale = -1.0f;
@@ -1840,6 +1841,10 @@ public class MediaPlayer14
 				frameUpdateInfo.VOffset = 1.0f - mVOffset;
 			}
 
+			// updateTexImage binds an external texture to active texture unit
+			// make sure to unbind it to prevent state leak
+			GLES20.glBindTexture(GL_TEXTURE_EXTERNAL_OES, 0);
+									
 			return frameUpdateInfo;
 		}
 	};
diff --git a/Engine/Build/Android/Java/src/com/epicgames/ue4/VideoDecoder.java b/Engine/Build/Android/Java/src/com/epicgames/ue4/VideoDecoder.java
--- a/Engine/Build/Android/Java/src/com/epicgames/ue4/VideoDecoder.java
+++ b/Engine/Build/Android/Java/src/com/epicgames/ue4/VideoDecoder.java
@@ -1333,6 +1333,7 @@ public class VideoDecoder
 		private int mTextureID = -1;
 		private float[] mTransformMatrix = new float[16];
 		private boolean mTextureSizeChanged = true;
+		private int GL_TEXTURE_EXTERNAL_OES = 0x8D65;
 
 		private float mUScale = 1.0f;
 		private float mVScale = -1.0f;
@@ -1482,6 +1483,11 @@ public class VideoDecoder
 			Entry<Double,Double> frameTimeDuration = getFrameTimeDuration(getFrameTimestamp());
 			frameUpdateInfo.Timestamp = frameTimeDuration.getKey();
 			frameUpdateInfo.Duration = frameTimeDuration.getValue();
+
+			// updateTexImage binds an external texture to active texture unit
+			// make sure to unbind it to prevent state leak
+			GLES20.glBindTexture(GL_TEXTURE_EXTERNAL_OES, 0);
+
 			return frameUpdateInfo;
 		}
 	};
diff --git a/Engine/Build/Android/Java/src/com/epicgames/ue4/WebViewControl.java b/Engine/Build/Android/Java/src/com/epicgames/ue4/WebViewControl.java
--- a/Engine/Build/Android/Java/src/com/epicgames/ue4/WebViewControl.java
+++ b/Engine/Build/Android/Java/src/com/epicgames/ue4/WebViewControl.java
@@ -1387,6 +1387,7 @@ class WebViewControl
 		private int mTextureID = -1;
 		private float[] mTransformMatrix = new float[16];
 		private boolean mTextureSizeChanged = true;
+		private int GL_TEXTURE_EXTERNAL_OES = 0x8D65;
 
 		private float mUScale = 1.0f;
 		private float mVScale = -1.0f;
@@ -1516,7 +1517,11 @@ class WebViewControl
 			// Get the latest video texture frame.
 			
 			mSurfaceTexture.updateTexImage();
-			
+
+			// updateTexImage binds an external texture to active texture unit
+			// make sure to unbind it to prevent state leak
+			GLES20.glBindTexture(GL_TEXTURE_EXTERNAL_OES, 0);
+						
 			return frameUpdateInfo;
 		}
 	};
diff --git a/Engine/Plugins/Media/AndroidCamera/Source/AndroidCamera/Java/src/com/epicgames/ue4/CameraPlayer14.java b/Engine/Plugins/Media/AndroidCamera/Source/AndroidCamera/Java/src/com/epicgames/ue4/CameraPlayer14.java
--- a/Engine/Plugins/Media/AndroidCamera/Source/AndroidCamera/Java/src/com/epicgames/ue4/CameraPlayer14.java
+++ b/Engine/Plugins/Media/AndroidCamera/Source/AndroidCamera/Java/src/com/epicgames/ue4/CameraPlayer14.java
@@ -1778,6 +1778,7 @@ public class CameraPlayer14
 		private int mTextureID = -1;
 		private float[] mTransformMatrix = new float[16];
 		private boolean mTextureSizeChanged = true;
+		private int GL_TEXTURE_EXTERNAL_OES = 0x8D65;
 
 		private float mScaleRotation00 = 1.0f;
 		private float mScaleRotation01 = 0.0f;
@@ -1996,6 +1997,11 @@ public class CameraPlayer14
 				frameUpdateInfo.VOffset = mVOffset;
 			}
 
+			
+			// updateTexImage binds an external texture to active texture unit
+			// make sure to unbind it to prevent state leak
+			GLES20.glBindTexture(GL_TEXTURE_EXTERNAL_OES, 0);
+			
 			return frameUpdateInfo;
 		}
 	};
