Fixed extra brace in AndroidCamera plugin code

#android

diff --git a/Engine/Plugins/Media/AndroidCamera/Source/AndroidCamera/Java/src/com/epicgames/unreal/CameraPlayer14.java b/Engine/Plugins/Media/AndroidCamera/Source/AndroidCamera/Java/src/com/epicgames/unreal/CameraPlayer14.java
--- a/Engine/Plugins/Media/AndroidCamera/Source/AndroidCamera/Java/src/com/epicgames/unreal/CameraPlayer14.java
+++ b/Engine/Plugins/Media/AndroidCamera/Source/AndroidCamera/Java/src/com/epicgames/unreal/CameraPlayer14.java
@@ -909,7 +909,7 @@ public class CameraPlayer14
 						}
 					}
 				}
-}			}
+			}
 
 			if (mUseOwnContext)
 			{
