Fix a copy-paste issue in android MediaPlayer14, WebView and CameraPlayer14

diff --git a/Engine/Build/Android/Java/src/com/epicgames/unreal/MediaPlayer14.java b/Engine/Build/Android/Java/src/com/epicgames/unreal/MediaPlayer14.java
--- a/Engine/Build/Android/Java/src/com/epicgames/unreal/MediaPlayer14.java
+++ b/Engine/Build/Android/Java/src/com/epicgames/unreal/MediaPlayer14.java
@@ -708,7 +708,7 @@ public class MediaPlayer14
 		if (null != mBitmapRenderer)
 		{
 			while (WaitOnBitmapRender) ;
-			releaseOESTextureRenderer();
+			releaseBitmapRenderer();
 		}
 		super.release();
 	}
diff --git a/Engine/Build/Android/Java/src/com/epicgames/unreal/VideoDecoder.java b/Engine/Build/Android/Java/src/com/epicgames/unreal/VideoDecoder.java
--- a/Engine/Build/Android/Java/src/com/epicgames/unreal/VideoDecoder.java
+++ b/Engine/Build/Android/Java/src/com/epicgames/unreal/VideoDecoder.java
@@ -244,7 +244,7 @@ public class VideoDecoder
 		if (null != mBitmapRenderer)
 		{
 			while (WaitOnBitmapRender) ;
-			releaseOESTextureRenderer();
+			releaseBitmapRenderer();
 		}
 
 		resetCodec();
diff --git a/Engine/Build/Android/Java/src/com/epicgames/unreal/WebViewControl.java b/Engine/Build/Android/Java/src/com/epicgames/unreal/WebViewControl.java
--- a/Engine/Build/Android/Java/src/com/epicgames/unreal/WebViewControl.java
+++ b/Engine/Build/Android/Java/src/com/epicgames/unreal/WebViewControl.java
@@ -228,7 +228,7 @@ class WebViewControl
 		if (null != mBitmapRenderer)
 		{
 			while (WaitOnBitmapRender) ;
-			releaseOESTextureRenderer();
+			releaseBitmapRenderer();
 		}
 
 		if (TargetCanvas != null)
diff --git a/Engine/Plugins/Media/AndroidCamera/Source/AndroidCamera/Java/src/com/epicgames/unreal/CameraPlayer14.java b/Engine/Plugins/Media/AndroidCamera/Source/AndroidCamera/Java/src/com/epicgames/unreal/CameraPlayer14.java
--- a/Engine/Plugins/Media/AndroidCamera/Source/AndroidCamera/Java/src/com/epicgames/unreal/CameraPlayer14.java
+++ b/Engine/Plugins/Media/AndroidCamera/Source/AndroidCamera/Java/src/com/epicgames/unreal/CameraPlayer14.java
@@ -706,7 +706,7 @@ public class CameraPlayer14
 		if (null != mBitmapRenderer)
 		{
 			while (WaitOnBitmapRender) ;
-			releaseOESTextureRenderer();
+			releaseBitmapRenderer();
 		}
 
 		if (null != mCamera)
