Fix a crash when playing a mp4 from UMG

diff --git a/Engine/Plugins/Media/AndroidCamera/Source/AndroidCamera/Private/Player/AndroidCameraPlayer.cpp b/Engine/Plugins/Media/AndroidCamera/Source/AndroidCamera/Private/Player/AndroidCameraPlayer.cpp
--- a/Engine/Plugins/Media/AndroidCamera/Source/AndroidCamera/Private/Player/AndroidCameraPlayer.cpp
+++ b/Engine/Plugins/Media/AndroidCamera/Source/AndroidCamera/Private/Player/AndroidCameraPlayer.cpp
@@ -303,7 +303,7 @@ void FAndroidCameraPlayer::SetGuid(const FGuid& Guid)
 // ==============================================================
 // Used by TickFetch
 
-static void DoUpdateExternalCameraSampleExecute(TWeakPtr<FJavaAndroidCameraPlayer, ESPMode::ThreadSafe> JavaCameraPlayerPtr, FGuid PlayerGuid)
+static void ConditionalInitializeVideoTexture_RenderThread(TWeakPtr<FJavaAndroidCameraPlayer, ESPMode::ThreadSafe> JavaCameraPlayerPtr, FGuid PlayerGuid)
 {
 	auto PinnedJavaCameraPlayer = JavaCameraPlayerPtr.Pin();
 
@@ -317,7 +317,8 @@ static void DoUpdateExternalCameraSampleExecute(TWeakPtr<FJavaAndroidCameraPlaye
 	{
 		const FRHITextureCreateDesc Desc =
 			FRHITextureCreateDesc::Create2D(TEXT("VideoTexture"), 1, 1, PF_R8G8B8A8)
-			.SetFlags(ETextureCreateFlags::External);
+			.SetFlags(ETextureCreateFlags::External)
+			.SetInitialState(ERHIAccess::SRVMask);
 
 		VideoTexture = RHICreateTexture(Desc);
 		PinnedJavaCameraPlayer->SetVideoTexture(VideoTexture);
@@ -334,6 +335,22 @@ static void DoUpdateExternalCameraSampleExecute(TWeakPtr<FJavaAndroidCameraPlaye
 		FPlatformMisc::LowLevelOutputDebugStringf(TEXT("Fetch RT: Created VideoTexture: %d - %s"), *reinterpret_cast<int32*>(VideoTexture->GetNativeResource()), *PlayerGuid.ToString());
 #endif
 	}
+}
+
+static void DoUpdateExternalCameraSampleExecute(TWeakPtr<FJavaAndroidCameraPlayer, ESPMode::ThreadSafe> JavaCameraPlayerPtr, FGuid PlayerGuid)
+{
+	auto PinnedJavaCameraPlayer = JavaCameraPlayerPtr.Pin();
+
+	if (!PinnedJavaCameraPlayer.IsValid())
+	{
+		return;
+	}
+
+	FTextureRHIRef VideoTexture = PinnedJavaCameraPlayer->GetVideoTexture();
+	if (VideoTexture == nullptr)
+	{
+		return;
+	}
 
 	int32 TextureId = *reinterpret_cast<int32*>(VideoTexture->GetNativeResource());
 	int32 CurrentFramePosition = 0;
@@ -369,6 +386,14 @@ struct FRHICommandUpdateExternalCameraSample final : public FRHICommand<FRHIComm
 	}
 };
 
+static void ConditionalInitializeVideoSample_RenderThread(TSharedRef<FAndroidCameraTextureSample, ESPMode::ThreadSafe> VideoSample)
+{
+	if (VideoSample->GetTexture() == nullptr || VideoSample->GetTexture()->GetSizeXY() != VideoSample->GetOutputDim())
+	{
+		VideoSample->InitializeTexture(FTimespan::Zero());
+	}
+}
+
 static void DoUpdateTextureCameraSampleExecute(TWeakPtr<FJavaAndroidCameraPlayer, ESPMode::ThreadSafe> JavaCameraPlayerPtr, TWeakPtr<FMediaSamples, ESPMode::ThreadSafe> SamplesPtr,
 	TSharedRef<FAndroidCameraTextureSample, ESPMode::ThreadSafe> VideoSample, FTimespan SampleTime)
 {
@@ -592,6 +617,7 @@ void FAndroidCameraPlayer::TickFetch(FTimespan DeltaTime, FTimespan Timecode)
 		ENQUEUE_RENDER_COMMAND(AndroidCameraPlayerWriteVideoSample)(
 			[Params](FRHICommandListImmediate& RHICmdList)
 			{
+				ConditionalInitializeVideoTexture_RenderThread(Params.JavaCameraPlayerPtr, Params.PlayerGuid);
 				if (IsRunningRHIInSeparateThread())
 				{
 					new (RHICmdList.AllocCommand<FRHICommandUpdateExternalCameraSample>()) FRHICommandUpdateExternalCameraSample(Params.JavaCameraPlayerPtr, Params.PlayerGuid);
@@ -661,6 +687,7 @@ void FAndroidCameraPlayer::TickFetch(FTimespan DeltaTime, FTimespan Timecode)
 		ENQUEUE_RENDER_COMMAND(AndroidCameraPlayerWriteVideoSample)(
 			[Params](FRHICommandListImmediate& RHICmdList)
 			{
+				ConditionalInitializeVideoSample_RenderThread(Params.VideoSample);
 				if (IsRunningRHIInSeparateThread())
 				{
 					new (RHICmdList.AllocCommand<FRHICommandUpdateTextureCameraSample>()) FRHICommandUpdateTextureCameraSample(Params.JavaCameraPlayerPtr, Params.SamplesPtr, Params.VideoSample, Params.SampleTime);
diff --git a/Engine/Plugins/Media/AndroidMedia/Source/AndroidMedia/Private/Player/AndroidMediaPlayer.cpp b/Engine/Plugins/Media/AndroidMedia/Source/AndroidMedia/Private/Player/AndroidMediaPlayer.cpp
--- a/Engine/Plugins/Media/AndroidMedia/Source/AndroidMedia/Private/Player/AndroidMediaPlayer.cpp
+++ b/Engine/Plugins/Media/AndroidMedia/Source/AndroidMedia/Private/Player/AndroidMediaPlayer.cpp
@@ -463,7 +463,7 @@ void FAndroidMediaPlayer::SetGuid(const FGuid& Guid)
 // ==============================================================
 // Used by TickFetch
 
-static void DoUpdateExternalMediaSampleExecute(TWeakPtr<FJavaAndroidMediaPlayer, ESPMode::ThreadSafe> JavaMediaPlayerPtr, FGuid PlayerGuid)
+static void ConditionalInitializeVideoTexture_RenderThread(TWeakPtr<FJavaAndroidMediaPlayer, ESPMode::ThreadSafe> JavaMediaPlayerPtr, FGuid PlayerGuid)
 {
 	auto PinnedJavaMediaPlayer = JavaMediaPlayerPtr.Pin();
 
@@ -477,7 +477,8 @@ static void DoUpdateExternalMediaSampleExecute(TWeakPtr<FJavaAndroidMediaPlayer,
 	{
 		const FRHITextureCreateDesc Desc =
 			FRHITextureCreateDesc::Create2D(TEXT("VideoTexture"), 1, 1, PF_R8G8B8A8)
-			.SetFlags(ETextureCreateFlags::External);
+			.SetFlags(ETextureCreateFlags::External)
+			.SetInitialState(ERHIAccess::SRVMask);
 
 		VideoTexture = RHICreateTexture(Desc);
 		PinnedJavaMediaPlayer->SetVideoTexture(VideoTexture);
@@ -494,6 +495,22 @@ static void DoUpdateExternalMediaSampleExecute(TWeakPtr<FJavaAndroidMediaPlayer,
 		FPlatformMisc::LowLevelOutputDebugStringf(TEXT("Fetch RT: Created VideoTexture: %d - %s"), *reinterpret_cast<int32*>(VideoTexture->GetNativeResource()), *PlayerGuid.ToString());
 #endif
 	}
+}
+
+static void DoUpdateExternalMediaSampleExecute(TWeakPtr<FJavaAndroidMediaPlayer, ESPMode::ThreadSafe> JavaMediaPlayerPtr, FGuid PlayerGuid)
+{
+	auto PinnedJavaMediaPlayer = JavaMediaPlayerPtr.Pin();
+
+	if (!PinnedJavaMediaPlayer.IsValid())
+	{
+		return;
+	}
+
+	FTextureRHIRef VideoTexture = PinnedJavaMediaPlayer->GetVideoTexture();
+	if (VideoTexture == nullptr)
+	{
+		return;
+	}
 
 	int32 TextureId = *reinterpret_cast<int32*>(VideoTexture->GetNativeResource());
 	int32 CurrentFramePosition = 0;
@@ -535,6 +552,14 @@ struct FRHICommandUpdateExternalMediaSample final : public FRHICommand<FRHIComma
 	}
 };
 
+static void ConditionalInitializeVideoSample_RenderThread(TSharedRef<FAndroidMediaTextureSample, ESPMode::ThreadSafe> VideoSample)
+{
+	if (VideoSample->GetTexture() == nullptr || VideoSample->GetTexture()->GetSizeXY() != VideoSample->GetOutputDim())
+	{
+		VideoSample->InitializeTexture(FTimespan::Zero());
+	}
+}
+
 static void DoUpdateTextureMediaSampleExecute(TWeakPtr<FJavaAndroidMediaPlayer, ESPMode::ThreadSafe> JavaMediaPlayerPtr, TWeakPtr<FMediaSamples, ESPMode::ThreadSafe> SamplesPtr,
 	TSharedRef<FAndroidMediaTextureSample, ESPMode::ThreadSafe> VideoSample)
 {
@@ -748,6 +773,7 @@ void FAndroidMediaPlayer::TickFetch(FTimespan DeltaTime, FTimespan /*Timecode*/)
 		ENQUEUE_RENDER_COMMAND(AndroidMediaPlayerWriteVideoSample)(
 			[Params](FRHICommandListImmediate& RHICmdList)
 			{
+				ConditionalInitializeVideoTexture_RenderThread(Params.JavaMediaPlayerPtr, Params.PlayerGuid);
 				if (IsRunningRHIInSeparateThread())
 				{
 					new (RHICmdList.AllocCommand<FRHICommandUpdateExternalMediaSample>()) FRHICommandUpdateExternalMediaSample(Params.JavaMediaPlayerPtr, Params.PlayerGuid);
@@ -811,6 +837,7 @@ void FAndroidMediaPlayer::TickFetch(FTimespan DeltaTime, FTimespan /*Timecode*/)
 		ENQUEUE_RENDER_COMMAND(AndroidMediaPlayerWriteVideoSample)(
 			[Params](FRHICommandListImmediate& RHICmdList)
 			{
+				ConditionalInitializeVideoSample_RenderThread(Params.VideoSample);
 				if (IsRunningRHIInSeparateThread())
 				{
 					new (RHICmdList.AllocCommand<FRHICommandUpdateTextureMediaSample>()) FRHICommandUpdateTextureMediaSample(Params.JavaMediaPlayerPtr, Params.SamplesPtr, Params.VideoSample);
diff --git a/Engine/Plugins/Media/ElectraUtil/Source/ElectraSamples/Private/ElectraTextureSampleAndroid.cpp b/Engine/Plugins/Media/ElectraUtil/Source/ElectraSamples/Private/ElectraTextureSampleAndroid.cpp
--- a/Engine/Plugins/Media/ElectraUtil/Source/ElectraSamples/Private/ElectraTextureSampleAndroid.cpp
+++ b/Engine/Plugins/Media/ElectraUtil/Source/ElectraSamples/Private/ElectraTextureSampleAndroid.cpp
@@ -479,7 +479,8 @@ bool FElectraTextureSample::Convert(FTexture2DRHIRef& InDstTexture, const FConve
 		// Make a source texture so we can convert from it...
 		const FRHITextureCreateDesc Desc =
 			FRHITextureCreateDesc::Create2D(TEXT("FMediaTextureResource"), SampleDim, VideoDecoderOutputAndroid->GetFormat())
-			.SetFlags(TexCreate_Dynamic);
+			.SetFlags(TexCreate_Dynamic)
+			.SetInitialState(ERHIAccess::SRVMask);
 		InputTexture = RHICreateTexture(Desc);
 		if (!InputTexture.IsValid())
 		{
