Fix up more references to GameActivity._activity

#android

diff --git a/Engine/Plugins/Media/AndroidCamera/Source/AndroidCamera/Java/src/com/epicgames/unreal/CameraPlayer14.java b/Engine/Plugins/Media/AndroidCamera/Source/AndroidCamera/Java/src/com/epicgames/unreal/CameraPlayer14.java
--- a/Engine/Plugins/Media/AndroidCamera/Source/AndroidCamera/Java/src/com/epicgames/unreal/CameraPlayer14.java
+++ b/Engine/Plugins/Media/AndroidCamera/Source/AndroidCamera/Java/src/com/epicgames/unreal/CameraPlayer14.java
@@ -130,7 +130,7 @@ public class CameraPlayer14
 		*/
 
 		// check if camera supported
-		if (GameActivity._activity.getPackageManager().hasSystemFeature(PackageManager.FEATURE_CAMERA_ANY))
+		if (GameActivity.Get().getPackageManager().hasSystemFeature(PackageManager.FEATURE_CAMERA_ANY))
 		{
 			CameraSupported = true;
 			GameActivity.Log.debug("Camera supported");
@@ -171,7 +171,7 @@ public class CameraPlayer14
 		if (Build.VERSION.SDK_INT >= Build.VERSION_CODES.ICE_CREAM_SANDWICH)
 		{
 			// register lifecycle callbacks to unlock / reconnect to camera if application goes to background
-			GameActivity._activity.getApplication().registerActivityLifecycleCallbacks(
+			GameActivity.Get().getApplication().registerActivityLifecycleCallbacks(
 				new Application.ActivityLifecycleCallbacks() {
 					@Override
 					public void onActivityCreated(Activity activity, Bundle savedInstanceState) {}
@@ -1383,7 +1383,7 @@ public class CameraPlayer14
 
 			mSurfaceTexture.getTransformMatrix(mTransformMatrix);
 
-			int degrees = (GameActivity._activity.getCurrentDeviceRotationDegree() + CameraRotationOffset) % 360;
+			int degrees = (GameActivity.Get().getCurrentDeviceRotationDegree() + CameraRotationOffset) % 360;
 			if (degrees == 0) {
 				mScaleRotation00 = mTransformMatrix[1];
 				mScaleRotation01 = mTransformMatrix[5];
@@ -1923,7 +1923,7 @@ public class CameraPlayer14
 			GameActivity.Log.debug(mTransformMatrix[12] + ", " + mTransformMatrix[13] + ", " + mTransformMatrix[14] + ", " + mTransformMatrix[15]);
 			*/
 
-			int degrees = (GameActivity._activity.getCurrentDeviceRotationDegree() + CameraRotationOffset) % 360;
+			int degrees = (GameActivity.Get().getCurrentDeviceRotationDegree() + CameraRotationOffset) % 360;
 			if (degrees == 0) {
 				mScaleRotation00 = mTransformMatrix[1];
 				mScaleRotation01 = mTransformMatrix[5];
diff --git a/Engine/Plugins/Runtime/GoogleCloudMessaging/Source/Java/RemoteNotificationsInstanceIDListener.java b/Engine/Plugins/Runtime/GoogleCloudMessaging/Source/Java/RemoteNotificationsInstanceIDListener.java
--- a/Engine/Plugins/Runtime/GoogleCloudMessaging/Source/Java/RemoteNotificationsInstanceIDListener.java
+++ b/Engine/Plugins/Runtime/GoogleCloudMessaging/Source/Java/RemoteNotificationsInstanceIDListener.java
@@ -10,7 +10,7 @@ public class RemoteNotificationsInstanceIDListener extends InstanceIDListenerSer
 	@Override
 	public void onTokenRefresh()
 	{
-		Intent intent = new Intent( GameActivity._activity.getApplicationContext(), RemoteNotificationsRegistrationIntentService.class );
+		Intent intent = new Intent( GameActivity.Get().getApplicationContext(), RemoteNotificationsRegistrationIntentService.class );
 		startService( intent );
 	}
 }
diff --git a/Engine/Plugins/Runtime/GoogleCloudMessaging/Source/Java/RemoteNotificationsListener.java b/Engine/Plugins/Runtime/GoogleCloudMessaging/Source/Java/RemoteNotificationsListener.java
--- a/Engine/Plugins/Runtime/GoogleCloudMessaging/Source/Java/RemoteNotificationsListener.java
+++ b/Engine/Plugins/Runtime/GoogleCloudMessaging/Source/Java/RemoteNotificationsListener.java
@@ -49,6 +49,6 @@ public class RemoteNotificationsListener extends GcmListenerService
 		notificationManager.notify((int)System.currentTimeMillis(), notification);		
 		*/
 		
-		GameActivity._activity.nativeGCMReceivedRemoteNotification( message );
+		GameActivity.Get().nativeGCMReceivedRemoteNotification( message );
 	}
 }
diff --git a/Engine/Plugins/Runtime/GoogleCloudMessaging/Source/Java/RemoteNotificationsRegistrationIntentService.java b/Engine/Plugins/Runtime/GoogleCloudMessaging/Source/Java/RemoteNotificationsRegistrationIntentService.java
--- a/Engine/Plugins/Runtime/GoogleCloudMessaging/Source/Java/RemoteNotificationsRegistrationIntentService.java
+++ b/Engine/Plugins/Runtime/GoogleCloudMessaging/Source/Java/RemoteNotificationsRegistrationIntentService.java
@@ -25,22 +25,22 @@ public class RemoteNotificationsRegistrationIntentService extends IntentService
   {
     try
     {
-      InstanceID instanceID = InstanceID.getInstance( GameActivity._activity.getApplicationContext() );
-      String Token = instanceID.getToken( GameActivity._activity.GCMSenderId, GoogleCloudMessaging.INSTANCE_ID_SCOPE, null );
+      InstanceID instanceID = InstanceID.getInstance( GameActivity.Get().getApplicationContext() );
+      String Token = instanceID.getToken( GameActivity.Get().GCMSenderId, GoogleCloudMessaging.INSTANCE_ID_SCOPE, null );
 
       // subscribe to default topic: /topics/global
       subscribeTopics( Token );
 
-      GameActivity._activity.nativeGCMRegisteredForRemoteNotifications( Token );
+      GameActivity.Get().nativeGCMRegisteredForRemoteNotifications( Token );
     }
     catch( Exception e )
     {
-      GameActivity._activity.nativeGCMFailedToRegisterForRemoteNotifications( "Failed to complete token refresh" );
+      GameActivity.Get().nativeGCMFailedToRegisterForRemoteNotifications( "Failed to complete token refresh" );
     }
   }
 
   private void subscribeTopics( String token ) throws IOException {
-    GcmPubSub pubSub = GcmPubSub.getInstance( GameActivity._activity.getApplicationContext() );
+    GcmPubSub pubSub = GcmPubSub.getInstance( GameActivity.Get().getApplicationContext() );
     pubSub.subscribe( token, "/topics/global", null );
   }
 }
