Updated CameraPlayer14 to Camera2 API.
Added ImageReader to the BitmapRenderer.
Fixed the support of device orientation.
Fixed the support of App moving to background/foreground.
Added minor fixes and improvement to MediaPlayer14.

diff --git a/Engine/Build/Android/Java/src/com/epicgames/unreal/BitmapRendererLegacy.java b/Engine/Build/Android/Java/src/com/epicgames/unreal/BitmapRendererLegacy.java
--- a/Engine/Build/Android/Java/src/com/epicgames/unreal/BitmapRendererLegacy.java
+++ b/Engine/Build/Android/Java/src/com/epicgames/unreal/BitmapRendererLegacy.java
@@ -31,6 +31,7 @@ class BitmapRendererLegacy implements BitmapRenderer, SurfaceTexture.OnFrameAvai
     private boolean mUseOwnContext = true;
     private boolean mVulkanRenderer = false;
     private boolean mSwizzlePixels = false;
+	private boolean mCorrectUV = false;
 
     private int GL_TEXTURE_EXTERNAL_OES = 0x8D65;
 
@@ -45,10 +46,11 @@ class BitmapRendererLegacy implements BitmapRenderer, SurfaceTexture.OnFrameAvai
 
     private boolean mCreatedEGLDisplay = false;
 
-    public BitmapRendererLegacy(boolean swizzlePixels, boolean isVulkan)
+    public BitmapRendererLegacy(boolean swizzlePixels, boolean isVulkan, boolean correctUV)
     {
         mSwizzlePixels = swizzlePixels;
         mVulkanRenderer = isVulkan;
+		mCorrectUV = correctUV;
 
         mEglSurface = EGL14.EGL_NO_SURFACE;
         mEglContext = EGL14.EGL_NO_CONTEXT;
@@ -543,29 +545,32 @@ class BitmapRendererLegacy implements BitmapRenderer, SurfaceTexture.OnFrameAvai
 
         mSurfaceTexture.getTransformMatrix(mTransformMatrix);
 
-        float UMin = mTransformMatrix[12];
-        float UMax = UMin + mTransformMatrix[0];
-        float VMin = mTransformMatrix[13];
-        float VMax = VMin + mTransformMatrix[5];
-
-        if (mTriangleVerticesData[2] != UMin ||
-                mTriangleVerticesData[6] != UMax ||
-                mTriangleVerticesData[11] != VMin ||
-                mTriangleVerticesData[3] != VMax)
-        {
-            //GameActivity.Log.debug("Matrix:");
-            //GameActivity.Log.debug(mTransformMatrix[0] + " " + mTransformMatrix[1] + " " + mTransformMatrix[2] + " " + mTransformMatrix[3]);
-            //GameActivity.Log.debug(mTransformMatrix[4] + " " + mTransformMatrix[5] + " " + mTransformMatrix[6] + " " + mTransformMatrix[7]);
-            //GameActivity.Log.debug(mTransformMatrix[8] + " " + mTransformMatrix[9] + " " + mTransformMatrix[10] + " " + mTransformMatrix[11]);
-            //GameActivity.Log.debug(mTransformMatrix[12] + " " + mTransformMatrix[13] + " " + mTransformMatrix[14] + " " + mTransformMatrix[15]);
-            mTriangleVerticesData[ 2] = mTriangleVerticesData[10] = UMin;
-            mTriangleVerticesData[ 6] = mTriangleVerticesData[14] = UMax;
-            mTriangleVerticesData[11] = mTriangleVerticesData[15] = VMin;
-            mTriangleVerticesData[ 3] = mTriangleVerticesData[ 7] = VMax;
-            mTriangleVerticesDirty = true;
-            //GameActivity.Log.debug("U = " + mTriangleVerticesData[2] + ", " + mTriangleVerticesData[6]);
-            //GameActivity.Log.debug("V = " + mTriangleVerticesData[11] + ", " + mTriangleVerticesData[3]);
-        }
+		if (mCorrectUV)
+		{
+			float UMin = mTransformMatrix[12];
+			float UMax = UMin + mTransformMatrix[0];
+			float VMin = mTransformMatrix[13];
+			float VMax = VMin + mTransformMatrix[5];
+
+			if (mTriangleVerticesData[2] != UMin ||
+					mTriangleVerticesData[6] != UMax ||
+					mTriangleVerticesData[11] != VMin ||
+					mTriangleVerticesData[3] != VMax)
+			{
+				//GameActivity.Log.debug("Matrix:");
+				//GameActivity.Log.debug(mTransformMatrix[0] + " " + mTransformMatrix[1] + " " + mTransformMatrix[2] + " " + mTransformMatrix[3]);
+				//GameActivity.Log.debug(mTransformMatrix[4] + " " + mTransformMatrix[5] + " " + mTransformMatrix[6] + " " + mTransformMatrix[7]);
+				//GameActivity.Log.debug(mTransformMatrix[8] + " " + mTransformMatrix[9] + " " + mTransformMatrix[10] + " " + mTransformMatrix[11]);
+				//GameActivity.Log.debug(mTransformMatrix[12] + " " + mTransformMatrix[13] + " " + mTransformMatrix[14] + " " + mTransformMatrix[15]);
+				mTriangleVerticesData[ 2] = mTriangleVerticesData[10] = UMin;
+				mTriangleVerticesData[ 6] = mTriangleVerticesData[14] = UMax;
+				mTriangleVerticesData[11] = mTriangleVerticesData[15] = VMin;
+				mTriangleVerticesData[ 3] = mTriangleVerticesData[ 7] = VMax;
+				mTriangleVerticesDirty = true;
+				//GameActivity.Log.debug("U = " + mTriangleVerticesData[2] + ", " + mTriangleVerticesData[6]);
+				//GameActivity.Log.debug("V = " + mTriangleVerticesData[11] + ", " + mTriangleVerticesData[3]);
+			}
+		}
 
         if (null != destData)
         {
diff --git a/Engine/Build/Android/Java/src/com/epicgames/unreal/MediaPlayer14.java b/Engine/Build/Android/Java/src/com/epicgames/unreal/MediaPlayer14.java
--- a/Engine/Build/Android/Java/src/com/epicgames/unreal/MediaPlayer14.java
+++ b/Engine/Build/Android/Java/src/com/epicgames/unreal/MediaPlayer14.java
@@ -757,7 +757,7 @@ public class MediaPlayer14
 	{
 		releaseBitmapRenderer();
 
-		GameActivity.Log.debug("CREATE BITMAPRENDERER, VIDEO SIZE: " + getVideoWidth() + " x " + getVideoHeight() + ", BRIR: " + usingBitmapRendererImageReader + ", Vulkan: " + VulkanRenderer);
+		GameActivity.Log.debug("MediaPlayer - CREATE BITMAPRENDERER, VIDEO SIZE: " + getVideoWidth() + " x " + getVideoHeight() + ", BRIR: " + usingBitmapRendererImageReader + ", Vulkan: " + VulkanRenderer);
 		
 		if (usingBitmapRendererImageReader)
 		{
@@ -766,7 +766,7 @@ public class MediaPlayer14
 		}
 		else
 		{
-			mBitmapRenderer = new BitmapRendererLegacy(SwizzlePixels, VulkanRenderer);
+			mBitmapRenderer = new BitmapRendererLegacy(SwizzlePixels, VulkanRenderer, true);
 		}
 
 		if (!mBitmapRenderer.isValid())
@@ -821,7 +821,7 @@ public class MediaPlayer14
 		}
 	}
 
-	public BitmapRendererImageReader.FrameData getVideoLastMediaFrameData()
+	public BitmapRenderer.FrameData getVideoLastMediaFrameData()
 	{
 		if (usingBitmapRendererImageReader)
 		{
@@ -829,7 +829,7 @@ public class MediaPlayer14
 			if (mBitmapRenderer != null)
 			{
 				WaitOnBitmapRender = true;
-				BitmapRendererImageReader.FrameData data = mBitmapRenderer.updateFrameData(updateParam);
+				BitmapRenderer.FrameData data = mBitmapRenderer.updateFrameData(updateParam);
 				WaitOnBitmapRender = false;
 				return data;
 			}
@@ -850,15 +850,12 @@ public class MediaPlayer14
 			WaitOnBitmapRender = true;
 
 			updateParam.destTexture = 0;
-			BitmapRendererImageReader.FrameData data = mBitmapRenderer.updateFrameData(updateParam);
+			BitmapRenderer.FrameData data = mBitmapRenderer.updateFrameData(updateParam);
 
 			WaitOnBitmapRender = false;
 			return data != null ? data.binaryBuffer : null;
 		}
-		else
-		{
-			return null;
-		}
+		return null;
 	}
 
 	public boolean getVideoLastFrame(int destTexture)
@@ -875,15 +872,12 @@ public class MediaPlayer14
 			WaitOnBitmapRender = true;
 
 			updateParam.destTexture = destTexture;
-			BitmapRendererImageReader.FrameData data = mBitmapRenderer.updateFrameData(updateParam);
+			BitmapRenderer.FrameData data = mBitmapRenderer.updateFrameData(updateParam);
 
 			WaitOnBitmapRender = false;
-			return data.frameReady;
-		}
-		else
-		{
-			return false;
+			return data != null ? data.frameReady : false;
 		}
+		return false;
 	}
 
 	// ======================================================================================
diff --git a/Engine/Plugins/Media/AndroidCamera/Source/AndroidCamera/Java/src/com/epicgames/unreal/CameraPlayer14.java b/Engine/Plugins/Media/AndroidCamera/Source/AndroidCamera/Java/src/com/epicgames/unreal/CameraPlayer14.java
--- a/Engine/Plugins/Media/AndroidCamera/Source/AndroidCamera/Java/src/com/epicgames/unreal/CameraPlayer14.java
+++ b/Engine/Plugins/Media/AndroidCamera/Source/AndroidCamera/Java/src/com/epicgames/unreal/CameraPlayer14.java
@@ -2,28 +2,38 @@
 
 package com.epicgames.unreal;
 
+import android.Manifest;
 import android.app.Activity;
 import android.app.Application;
-import android.content.res.AssetFileDescriptor;
-import android.content.res.AssetManager;
+import android.content.Context;
 import android.content.pm.PackageManager;
-import android.content.res.Configuration;
-import android.graphics.Rect;
+import android.hardware.camera2.params.StreamConfigurationMap;
+import android.media.Image;
+import android.media.ImageReader;
 import android.opengl.*;
 import android.os.Build;
 import android.os.Bundle;
-import android.view.Surface;
-import java.io.File;
-import java.io.IOException;
-import java.io.FileOutputStream;
 import java.io.FileNotFoundException;
-import java.io.RandomAccessFile;
+import java.io.FileOutputStream;
+import java.io.IOException;
 import java.lang.RuntimeException;
-import android.hardware.Camera;
+import java.nio.ByteBuffer;
 import java.util.List;
 import java.util.ArrayList;
-import java.util.Random;
-import android.util.Log;
+import java.util.concurrent.Executor;
+
+import android.graphics.ImageFormat;
+import android.hardware.camera2.*;
+import android.os.Handler;
+import android.os.HandlerThread;
+import android.util.Range;
+import android.util.SparseIntArray;
+import android.view.Surface;
+
+import androidx.activity.result.ActivityResultLauncher;
+import androidx.activity.result.contract.ActivityResultContracts;
+import androidx.annotation.NonNull;
+import androidx.core.app.ActivityCompat;
 
 // from permission_library
 import com.google.vr.sdk.samples.permission.PermissionHelper;
@@ -40,10 +50,23 @@ public class CameraPlayer14
 
 	private boolean CameraSupported = false;
 	private boolean mVideoEnabled = true;
-	
-	private Camera mCamera = null;
+
+    private HandlerThread mBackgroundThread;
+    private Handler mBackgroundHandler;
+
+	private CameraManager mCameraManager = null;
+	private CameraDevice mCamera = null;
+	private CaptureRequest.Builder mPreviewRequestBuilder = null;
+	private boolean mAutoStartPreview = false;
+	private CameraCaptureSession mCameraCaptureSession = null;
+	private final int mCameraFormat = ImageFormat.YUV_420_888; //ImageFormat.JPEG;
+
+	private int DesiredCameraWidth = 0;
+	private int DesiredCameraHeight = 0;
+	private int DesiredCameraFPS = 0;
+	private String DesiredCameraId;
+
 	private CameraStates CameraState = CameraStates.PREINIT;
-	private int CameraId = 0;
 	private int CameraWidth = 0;
 	private int CameraHeight = 0;
 	private int CameraFPS = 0;
@@ -52,8 +75,9 @@ public class CameraPlayer14
 	private int CameraOrientation = 0;
 	private int CameraRotationOffset = 0;
 
-	private int iFrontId = 0;
-	private int iBackId = 0;
+	private String CameraId;
+	private String CameraFrontId;
+	private String CameraBackId;
 
 	private boolean SwizzlePixels = true;
 	private boolean VulkanRenderer = false;
@@ -66,6 +90,7 @@ public class CameraPlayer14
 	private OESTextureRenderer mOESTextureRenderer = null;
 
 	public class FrameUpdateInfo {
+		public BitmapRenderer.FrameData BRFrameData = null;
 		public java.nio.Buffer Buffer;
 		public int CurrentPosition;
 		public boolean FrameReady;
@@ -149,90 +174,91 @@ public class CameraPlayer14
 			CameraSupported = false;
 			GameActivity.Log.debug("Camera permission not granted");
 			return;
-		}
-		
+		}		
+
 		// discover available cameras and remember front and back IDs if available
-		Camera.CameraInfo cameraInfo = new Camera.CameraInfo();
-		int cameraCount = Camera.getNumberOfCameras();
-		for (int index = 0; index < cameraCount; index++) {
-			Camera.getCameraInfo(index, cameraInfo);
-			if (cameraInfo.facing == Camera.CameraInfo.CAMERA_FACING_FRONT)
+		mCameraManager = (CameraManager) GameActivity.Get().getSystemService(Context.CAMERA_SERVICE);
+		try {
+			String[] ids = mCameraManager.getCameraIdList();
+			if (ids.length <= 0)
 			{
-				iFrontId = index;
-			} else if (cameraInfo.facing == Camera.CameraInfo.CAMERA_FACING_BACK)
+				CameraSupported = false;
+				GameActivity.Log.debug("CameraPlayer14: No camera is in the device.");
+				return;
+			}
+
+			CameraFrontId = null;
+			CameraBackId = null;
+			for (String id : ids)
+			{
+				CameraCharacteristics characteristics = mCameraManager.getCameraCharacteristics(id);
+				int facing = characteristics.get(CameraCharacteristics.LENS_FACING);
+				switch (facing)
+				{
+					case CameraCharacteristics.LENS_FACING_FRONT:
+						CameraFrontId = id;
+						break;
+					case CameraCharacteristics.LENS_FACING_BACK:
+						CameraBackId = id;
+						break;
+				}
+				if (CameraFrontId != null && CameraBackId != null)
+				{
+					break;
+				}
+			}
+			
+			if (CameraFrontId == null && CameraBackId == null)
 			{
-				iBackId = index;
+				CameraSupported = false;
+				GameActivity.Log.debug("CameraPlayer14: No camera was found suitable for use.");
+				return;
 			}
 		}
+		catch (CameraAccessException e)
+		{
+			CameraSupported = false;
+			GameActivity.Log.error("CameraPlayer14: furtherInit exception - " + e.getLocalizedMessage());
+			return;
+		}
 
-		GameActivity.Log.debug("CameraPlayer14: frontId=" + iFrontId + ", backId=" + iBackId);
+		GameActivity.Log.debug("CameraPlayer14: frontId=" + CameraFrontId + ", backId=" + CameraBackId);
+	}
 
-		// not really necessary, but just in case...
-		if (Build.VERSION.SDK_INT >= Build.VERSION_CODES.ICE_CREAM_SANDWICH)
+	private void startBackgroundThread()
+	{
+		/* this is to allow the camera opperations to run on a separate thread and avoid blocking the UI*/
+		if (mBackgroundThread == null)
 		{
-			// register lifecycle callbacks to unlock / reconnect to camera if application goes to background
-			GameActivity.Get().getApplication().registerActivityLifecycleCallbacks(
-				new Application.ActivityLifecycleCallbacks() {
-					@Override
-					public void onActivityCreated(Activity activity, Bundle savedInstanceState) {}
-					@Override
-					public void onActivityStarted(Activity activity) {}
-					@Override
-					public void onActivitySaveInstanceState(Activity activity, Bundle outState) {}
-					@Override
-					public void onActivityPaused(Activity activity)
-					{
-						if (null != mCamera)
-						{
-							try {
-								if (CameraState == CameraStates.START)
-								{
-									GameActivity.Log.debug("CameraPlayer14: onPause: stop preview");
-									mCamera.stopPreview();
-								}
-								mCamera.release();
-								mCamera = null;
-								GameActivity.Log.debug("CameraPlayer14: onPause: released camera");
-							}
-							catch (RuntimeException e)
-							{
-								GameActivity.Log.debug("CameraPlayer14: onPause: failed to unlocked camera: " + e);
-							}
-						}
-					}
-					@Override
-					public void onActivityResumed(Activity activity)
-					{
-						if (CameraState == CameraStates.START || CameraState == CameraStates.PAUSE)
-						{
-							try {
-								mCamera = Camera.open(CameraId);
-								GameActivity.Log.debug("CameraPlayer14: onResume: opened camera");
-
-								if (CameraState == CameraStates.START)
-								{
-									setVideoEnabled(mVideoEnabled);
-								}
-							}
-							catch (Exception e)
-							{
-								GameActivity.Log.debug("CameraPlayer14: onResume: failed to reconnect camera: " + e);
-							}
-						}
-					}
-					@Override
-					public void onActivityStopped(Activity activity) 
-					{ 
-					}
-					@Override
-					public void onActivityDestroyed(Activity activity)
-					{
-                    }
-				});
+			mBackgroundThread = new HandlerThread("Camera2Handler");
+			mBackgroundThread.start();
+			mBackgroundHandler = new Handler(mBackgroundThread.getLooper());
 		}
 	}
 
-	public int getVideoWidth()
+	private void stopBackgroundThread()
+	{
+		if (mBackgroundThread == null)
+		{
+			mBackgroundHandler = null;
+			return;
+		}
+		
+		mBackgroundThread.quitSafely();
+		try
+		{
+			mBackgroundThread.join();
+		}
+		catch (InterruptedException e)
+		{
+			e.printStackTrace();
+		}
+
+		mBackgroundThread = null;
+		mBackgroundHandler = null;
+	}
+	
+	public synchronized int getVideoWidth()
 	{
 //		GameActivity.Log.debug("CameraPlayer14: getVideoWidth: " + CameraWidth);
 		return (CameraState == CameraStates.PREPARED ||
@@ -241,7 +267,7 @@ public class CameraPlayer14
 				CameraState == CameraStates.STOP) ? CameraWidth : 0;
 	}
 
-	public int getVideoHeight()
+	public synchronized int getVideoHeight()
 	{
 //		GameActivity.Log.debug("CameraPlayer14: getVideoHeight: " + CameraHeight);
 		return (CameraState == CameraStates.PREPARED ||
@@ -256,7 +282,7 @@ public class CameraPlayer14
 		return CameraFPS;
 	}
 
-	public void prepare()
+	public synchronized void prepare()
 	{
 		if (null == mCamera)
 		{
@@ -268,7 +294,7 @@ public class CameraPlayer14
 		Prepared = true;
 	}
 
-	public void prepareAsync()
+	public synchronized void prepareAsync()
 	{
 		if (null == mCamera)
 		{
@@ -280,7 +306,7 @@ public class CameraPlayer14
 		Prepared = true;
 	}
 	
-	public boolean isPlaying()
+	public synchronized boolean isPlaying()
 	{
 		return CameraState == CameraStates.START;
 	}
@@ -338,25 +364,14 @@ public class CameraPlayer14
 	{
 		if (!CameraSupported)
 		{
+			GameActivity.Log.error("CameraPlayer14: setDataSourceURL - CameraSupported is false.");
 			return false;
 		}
 
 //		GameActivity.Log.debug("setDataSource: Url = " + Url);
-		stop();
-		releaseOESTextureRenderer();
-		releaseBitmapRenderer();
-
-		synchronized(this)
-		{
-			Prepared = false;
-			Completed = false;
-		}
-		Looping = false;
-		audioTracks.clear();
-		videoTracks.clear();
-
-		// defaults		
-		int cameraId = 0;
+		
+		// defaults
+		String cameraId = "";
 		int width = 320;
 		int height = 480;
 		int fps = 30;
@@ -371,28 +386,44 @@ public class CameraPlayer14
 		if (UrlParts.length > 0)
 		{
 //			GameActivity.Log.debug("setDataSource: camera=" + UrlParts[0]);
-			if (UrlParts[0].equals("front"))
+			if (UrlParts[0].equals("front") && CameraFrontId != null)
 			{
-				cameraId = iFrontId;
+				cameraId = CameraFrontId;
 //				GameActivity.Log.debug("setDataSource: front=" + iFrontId);
 			}
-			else if (UrlParts[0].equals("back") || UrlParts[0].equals("rear"))
+			else if ((UrlParts[0].equals("back") || UrlParts[0].equals("rear")) && CameraBackId != null)
 			{
-				cameraId = iBackId;
+				cameraId = CameraBackId;
 //				GameActivity.Log.debug("setDataSource: back=" + iBackId);
 			}
 			else
 			{
 				try
 				{
-					cameraId = Integer.parseInt(UrlParts[0]);
+					int expectingCameraIndex = Integer.parseInt(UrlParts[0]);
+					String[] ids = mCameraManager.getCameraIdList();
+					if (expectingCameraIndex >= 0 && expectingCameraIndex < ids.length)
+					{
+						cameraId = ids[expectingCameraIndex];
+					}
+					else if (CameraBackId != null)
+					{
+						cameraId = CameraBackId;
+					}
 				}
-				catch (NumberFormatException e)
+				catch (NumberFormatException | CameraAccessException e)
 				{
+					throw new RuntimeException(e);
 				}
 //				GameActivity.Log.debug("setDataSource: back=" + iBackId);
 			}
 
+			if (cameraId.length() <= 0 && CameraBackId == null)
+			{
+				GameActivity.Log.error("CameraPlayer14: setDataSourceURL - Couldn't locate the camera specified in the URL.");
+				return false;
+			}
+			
 			for (int UrlPartIndex=1; UrlPartIndex < UrlParts.length; UrlPartIndex++)
 			{
 				String[] KeyValue = UrlParts[UrlPartIndex].split("=", 2);
@@ -422,73 +453,164 @@ public class CameraPlayer14
 				}
 			}
 		}
+
+		DesiredCameraId = cameraId;
+		DesiredCameraWidth = width;
+		DesiredCameraHeight = height;
+		DesiredCameraFPS = fps;
 		
-		try {
-			mCamera = Camera.open(cameraId);
+		return openCamera();
+	}
+
+	public boolean openCamera()
+	{
+		reset();
+
+		startBackgroundThread();
+
+		try
+		{
+			if (!CameraSupported)
+			{
+				GameActivity.Log.error("CameraPlayer14: openCamera - Camera is not supported.");
+				return false;
+			}
+
+			if (ActivityCompat.checkSelfPermission(GameActivity.Get(), Manifest.permission.CAMERA) != PackageManager.PERMISSION_GRANTED)
+			{
+				GameActivity.Log.error("CameraPlayer14: openCamera - Camera permission is not granted.");
+				return false;
+			}
+
+			GameActivity.Log.debug("CameraPlayer14: openCamera - id=" + DesiredCameraId);
+
+			CameraDevice.StateCallback CameraDeviceStateCallback = new CameraDevice.StateCallback()
+				{
+					@Override
+					public void onOpened(@NonNull CameraDevice camera)
+					{
+						mCamera = camera;
+						configureCamera(DesiredCameraId, DesiredCameraFPS, DesiredCameraWidth, DesiredCameraHeight);
+						prepare();
+					}
+
+					@Override
+					public void onDisconnected(@NonNull CameraDevice camera)
+					{
+						GameActivity.Log.error("CameraPlayer14: CameraDeviceStateCallback - Camera device disconnected.");
+					}
+
+					@Override
+					public void onError(@NonNull CameraDevice camera, int error)
+					{
+						GameActivity.Log.error("CameraPlayer14: CameraDeviceStateCallback - Camera device is in error state - code: " + error);
+					}
+				};
+
+			mCameraManager.openCamera(DesiredCameraId, CameraDeviceStateCallback, mBackgroundHandler);
 		}
-		catch (Exception e)
+		catch (CameraAccessException e)
 		{
-			mCamera = null;
-			GameActivity.Log.debug("setDataSource: could not open camera: " + cameraId);
+			GameActivity.Log.error("CameraPlayer14: openCamera - Camera access exception - " + e.getLocalizedMessage());
 			return false;
 		}
-		
-//		GameActivity.Log.debug("setDataSource: id=" + cameraId + ", width=" + width + ", height=" + height + ", fps=" + fps);
-		
-		Camera.Parameters param = mCamera.getParameters();
+
+		return true;
+	}
+
+	private boolean configureCamera(String cameraId, int fpsDesired, int widthDesired, int heightDesired)
+	{
+		Range<Integer>[] previewFPSRanges;
+		android.util.Size[] previewSizes;
+
+		try
+		{
+			mPreviewRequestBuilder = mCamera.createCaptureRequest(CameraDevice.TEMPLATE_PREVIEW);
+
+			CameraCharacteristics cameraChar = mCameraManager.getCameraCharacteristics(cameraId);
+
+			StreamConfigurationMap configs = cameraChar.get(CameraCharacteristics.SCALER_STREAM_CONFIGURATION_MAP);
+			previewSizes = configs.getOutputSizes(mCameraFormat);
+
+			previewFPSRanges = cameraChar.get(CameraCharacteristics.CONTROL_AE_AVAILABLE_TARGET_FPS_RANGES);
+
+			// fixup for orientation depends on camera facing
+			CameraRotationOffset = 0;
+			CameraOrientation = cameraChar.get(CameraCharacteristics.SENSOR_ORIENTATION);
+			int facing = cameraChar.get(CameraCharacteristics.LENS_FACING);
+			if (facing == CameraCharacteristics.LENS_FACING_FRONT)
+			{
+				CameraRotationOffset = (CameraOrientation + 90) % 360;
+			}
+			else if (facing == CameraCharacteristics.LENS_FACING_BACK)
+			{
+				CameraRotationOffset = (450 - CameraOrientation) % 360;
+			}
+		}
+		catch (CameraAccessException e)
+		{
+			GameActivity.Log.error("CameraPlayer14: configureCamera - camera access exception: " + e.getLocalizedMessage());
+			return false;
+		}
+
+		if (mPreviewRequestBuilder == null || previewFPSRanges == null || previewSizes == null)
+		{
+			GameActivity.Log.error("CameraPlayer14: configureCamera - unable to retrieve camera characteristics : " + cameraId);
+			return false;
+		}
+
+		GameActivity.Log.debug("CameraPlayer14: configureCamera - id=" + cameraId + ", width=" + widthDesired + ", height=" + heightDesired + ", fps=" + fpsDesired);
 
 		// find min and max fps allowed and matched frame rate range
-		float floatFps = fps;
-		fps *= 1000;
-		int fpsMin = fps;
-		int fpsMax = fps;
+		CameraFPS = fpsDesired;
+		CameraFPSMin = fpsDesired * 1000;
+		CameraFPSMax = CameraFPSMin;
 		float floatMinFps = 10000;
 		float floatMaxFps = -10000;
-		List<int[]> previewRates = param.getSupportedPreviewFpsRange();
 		boolean found = false;
 
 		// first look for an exact match
-		for (index = 0; index < previewRates.size(); ++index)
+		for (Range<Integer> range : previewFPSRanges)
 		{
-			int[] rates = previewRates.get(index);
-//			GameActivity.Log.debug("setDataSource: previewRates " + index + ": " + rates[0] + " - " + rates[1]);
-			float previewMin = (float)rates[0] / 1000.0f;
-			float previewMax = (float)rates[1] / 1000.0f;
+			//GameActivity.Log.debug("setDataSource: FPS Range: " + range.getLower() + " - " + range.getUpper());
+			float previewMin = range.getLower().floatValue() / 1000.0f;
+			float previewMax = range.getUpper().floatValue() / 1000.0f;
 			floatMinFps = previewMin < floatMinFps ? previewMin : floatMinFps;
 			floatMaxFps = previewMax > floatMaxFps ? previewMax : floatMaxFps;
-			if ((fps == rates[0] && fps == rates[1]))
+			if (fpsDesired == range.getLower() && fpsDesired == range.getUpper())
 			{
 				found = true;
-				fpsMin = rates[0];
-				fpsMax = rates[1];
+				CameraFPSMin = range.getLower();
+				CameraFPSMax = range.getUpper();
 			}
 		}
 		if (!found)
 		{
 			// look for range containing desired frame rate
-			for (index = 0; index < previewRates.size(); ++index)
+			for (Range<Integer> range : previewFPSRanges)
 			{
-				int[] rates = previewRates.get(index);
-				if (!found || (fps >= rates[0] && fps <= rates[1]))
+				if (!found || range.contains(fpsDesired))
 				{
 					found = true;
-					fpsMin = rates[0];
-					fpsMax = rates[1];
+					CameraFPSMin = range.getLower();
+					CameraFPSMax = range.getUpper();
 				}
 			}
 		}
 
 		// find best matched preview size and add track info
-		List<Camera.Size> previewSizes = param.getSupportedPreviewSizes();
-		if (previewSizes.size() > 0)
+		audioTracks.clear();
+		videoTracks.clear();
+
+		if (previewSizes.length > 0)
 		{
 			int bestWidth = -1;
 			int bestHeight = -1;
 			int bestScore = 0;
-			for (index = 0; index < previewSizes.size(); ++index)
+			for (int index = 0; index < previewSizes.length; ++index)
 			{
-				Camera.Size camSize = previewSizes.get(index);
-//				GameActivity.Log.debug("setDataSource: previewSize " + index + ": " + camSize.width + " x " + camSize.height);
+				android.util.Size camSize = previewSizes[index];
+				//GameActivity.Log.debug("CameraPlayer14: configureCamera - available previewSize " + index + ": " + camSize.getWidth() + " x " + camSize.getHeight());
 
 				VideoTrackInfo videoTrack = new VideoTrackInfo();
 				videoTrack.Index = index;
@@ -496,67 +618,129 @@ public class CameraPlayer14
 				videoTrack.DisplayName = "Video Track " + index + " (Stream " + index + ")";
 				videoTrack.Language = "und";
 				videoTrack.BitRate = 0;
-				videoTrack.Width = camSize.width;
-				videoTrack.Height = camSize.height;
-				videoTrack.FrameRate = floatFps;
+				videoTrack.Width = camSize.getWidth();
+				videoTrack.Height = camSize.getHeight();
+				videoTrack.FrameRate = fpsDesired;
 				videoTrack.FrameRateLow = floatMinFps;
 				videoTrack.FrameRateHigh = floatMaxFps;
 				videoTracks.add(videoTrack);
 
-				int diffWidth = camSize.width - width;
-				int diffHeight = camSize.height - height;
+				int diffWidth = camSize.getWidth() - widthDesired;
+				int diffHeight = camSize.getHeight() - heightDesired;
 				int score = diffWidth * diffWidth + diffHeight * diffHeight;
 				if (bestWidth == -1 || score < bestScore)
 				{
-					bestWidth = camSize.width;
-					bestHeight = camSize.height;
+					bestWidth = camSize.getWidth();
+					bestHeight = camSize.getHeight();
 					bestScore = score;
 				}
 			}
-			width = bestWidth;
-			height = bestHeight;
+			CameraWidth = bestWidth;
+			CameraHeight = bestHeight;
+			
+			GameActivity.Log.debug("CameraPlayer14: configureCamera - The preview size best matched: " + CameraWidth + " x " + CameraHeight);
 		}
 
 		// set camera parameters
-		param.setPreviewSize(width, height);
-		param.setPreviewFpsRange(fpsMin, fpsMax);
-		param.set("orientation", "landscape");
-		mCamera.setDisplayOrientation(0);
-		mCamera.setParameters(param);
+		mPreviewRequestBuilder.set(CaptureRequest.CONTROL_AE_TARGET_FPS_RANGE, new Range<>(CameraFPSMin, CameraFPSMax));
 
-		Camera.CameraInfo camInfo = new Camera.CameraInfo();
-		Camera.getCameraInfo(cameraId, camInfo);
-		CameraOrientation = camInfo.orientation;
-		CameraRotationOffset = 0;
+		// store our picks
+		CameraId = cameraId;
 
-		// fixup for orientation depends on camera facing
-		if (camInfo.facing == Camera.CameraInfo.CAMERA_FACING_FRONT)
+		GameActivity.Log.debug("CameraPlayer14: Camera configured - id=" + cameraId + ", width=" + CameraWidth + ", height=" + CameraHeight + ", fps=" + CameraFPS +
+				", min=" + CameraFPSMin + " max=" + CameraFPSMax + " orientation=" + CameraOrientation);
+
+		synchronized (this)
 		{
-			CameraRotationOffset = (CameraOrientation + 90) % 360;
+			CameraState = CameraStates.PREPARED;
 		}
-		else if (camInfo.facing == Camera.CameraInfo.CAMERA_FACING_BACK)
+		return true;
+	}
+
+	private void CreatePreviewSession(Surface previewSurface)
+	{
+		mPreviewRequestBuilder.addTarget(previewSurface);
+
+		List<Surface> listSurfaces = new ArrayList<Surface>();
+		listSurfaces.add(previewSurface);
+		try
 		{
-			CameraRotationOffset = (450 - CameraOrientation) % 360;
-		}
+			GameActivity.Log.debug("CameraPlayer14: creating capture session...");
 
-		// store our picks
-		CameraId = cameraId;
-		CameraWidth = width;
-		CameraHeight = height;
-		CameraFPSMin = fpsMin;
-		CameraFPSMax = fpsMax;
-		CameraFPS = fps / 1000;
+			mCamera.createCaptureSession(listSurfaces, new CameraCaptureSession.StateCallback()
+			{
+				@Override
+				public void onConfigured(@NonNull CameraCaptureSession session)
+				{
+					assert (session.getDevice() == mCamera);
+
+					synchronized (this)
+					{
+						mCameraCaptureSession = session;
+					}
 
-		GameActivity.Log.debug("setDataSource: id=" + cameraId + ", width=" + CameraWidth + ", height=" + CameraHeight + ", fps=" + CameraFPS + ", min=" + fpsMin + " max=" + fpsMax + " orientation=" + CameraOrientation);
+					GameActivity.Log.debug("CameraPlayer14: capture session created!");
 
-		CameraState = CameraStates.PREPARED;
-		return true;
+					if (mAutoStartPreview)
+					{
+						mAutoStartPreview = false;
+						startPreview();
+					}
+				}
+
+				@Override
+				public void onConfigureFailed(@NonNull CameraCaptureSession session)
+				{
+					GameActivity.Log.warn("CameraPlayer14: capture session creating failed!");
+				}
+			}, mBackgroundHandler);
+		}
+		catch (CameraAccessException e)
+		{
+			GameActivity.Log.error("CameraPlayer14: CreatePreviewSession - " + e.getLocalizedMessage());
+		}
+	}
+
+	private void startPreview()
+	{
+		if (mCameraCaptureSession == null)
+		{
+			GameActivity.Log.warn("CameraPlayer14: startPreview - Null CaptureSession");
+			return;
+		}
+		try
+		{
+			GameActivity.Log.debug("CameraPlayer14: Preview starting...");
+			mCameraCaptureSession.setRepeatingRequest(mPreviewRequestBuilder.build(), null, mBackgroundHandler);
+		}
+		catch (CameraAccessException e)
+		{
+			GameActivity.Log.error("CameraPlayer14: startPreview - " + e.getLocalizedMessage());
+		}
+	}
+
+	private void stopPreview()
+	{
+		if (mCameraCaptureSession == null)
+		{
+			GameActivity.Log.warn("CameraPlayer14: stopPreview - Null CaptureSession");
+			return;
+		}
+		try
+		{
+			mCameraCaptureSession.stopRepeating();
+			GameActivity.Log.debug("CameraPlayer14: Preview stopped");
+		}
+		catch (CameraAccessException e)
+		{
+			GameActivity.Log.error("CameraPlayer14: stopPreview - " + e.getLocalizedMessage());
+		}
 	}
 
 	public String getDataSourceURL()
 	{
-		String URL = "vidcap://" + ((CameraId == iFrontId) ? "front" : "rear") + "?width=" + CameraWidth + "?height=" + CameraHeight + "?fps=" + CameraFPS;
-		GameActivity.Log.debug("getDataSource: " + URL);
+		String URL = "vidcap://" + ((CameraId == CameraFrontId) ? "front" : "rear") + "?width=" + CameraWidth + "?height=" + CameraHeight + "?fps=" + CameraFPS;
+		GameActivity.Log.debug("CameraPlayer14: getDataSource - " + URL);
 		return URL;
 	}
 
@@ -567,7 +751,7 @@ public class CameraPlayer14
 			return;
 		}
 
-//		GameActivity.Log.debug("CameraPlayer14: setVideoEnabled: " + enabled);
+		GameActivity.Log.debug("CameraPlayer14: setVideoEnabled: " + enabled);
 		WaitOnBitmapRender = true;
 
 		mVideoEnabled = enabled;
@@ -575,45 +759,26 @@ public class CameraPlayer14
 		{
 			if (null != mOESTextureRenderer && null != mOESTextureRenderer.getSurface())
 			{
-				try
-				{
-					mCamera.setPreviewTexture(mOESTextureRenderer.getSurfaceTexture());
-				}
-				catch (IOException e)
-				{
-				}
+				CreatePreviewSession(mOESTextureRenderer.getSurface());
 			}
 
 			if (null != mBitmapRenderer && null != mBitmapRenderer.getSurface())
 			{
-				try
-				{
-					mCamera.setPreviewTexture(mBitmapRenderer.getSurfaceTexture());
-				}
-				catch (IOException e)
-				{
-				}
+				CreatePreviewSession(mBitmapRenderer.getSurface());
 			}
 		}
 		else
 		{
-			try
-			{
-				mCamera.setPreviewTexture(null);
-			}
-			catch (IOException e)
-			{
-			}
+			stopPreview();
 		}
 
 		WaitOnBitmapRender = false;
-		mCamera.startPreview();
 	}
-	
+
 	public void setAudioEnabled(boolean enabled)
 	{
 	}
-	
+
 	public boolean didResolutionChange()
 	{
 		if (null != mOESTextureRenderer)
@@ -633,14 +798,14 @@ public class CameraPlayer14
 		{
 			return mOESTextureRenderer.getExternalTextureId();
 		}
-		if (null != mBitmapRenderer)
+		if (null != mBitmapRenderer && (mBitmapRenderer instanceof BitmapRendererLegacy))
 		{
-			return mBitmapRenderer.getExternalTextureId();
+			return ((BitmapRendererLegacy)mBitmapRenderer).getExternalTextureId();
 		}
 		return -1;
 	}
 
-	public void start()
+	public synchronized void start()
 	{
 //		GameActivity.Log.debug("CameraPlayer14: start: " + CameraState);
 		if (CameraState == CameraStates.PREPARED)
@@ -650,33 +815,27 @@ public class CameraPlayer14
 
 		if (CameraState == CameraStates.PAUSE)
 		{
-			mCamera.startPreview();
+			startPreview();
 			CameraState = CameraStates.START;
 		}
 	}
 
-	public void pause()
+	public synchronized void pause()
 	{
 //		GameActivity.Log.debug("CameraPlayer14: pause: " + CameraState);
 		if (CameraState == CameraStates.START)
 		{
-			mCamera.stopPreview();
-			CameraState = CameraStates.PAUSE;
+			stopPreview();
 		}
+		CameraState = CameraStates.PAUSE;
 	}
 
-	public void stop()
+	public synchronized void stop()
 	{
 //		GameActivity.Log.debug("CameraPlayer14: stop: " + CameraState);
 		if (CameraState == CameraStates.START)
 		{
-			mCamera.stopPreview();
-			reset();
-		}
-		if (null != mCamera)
-		{
-			mCamera.release();
-			mCamera = null;
+			stopPreview();
 		}
 		CameraState = CameraStates.STOP;
 	}
@@ -695,31 +854,36 @@ public class CameraPlayer14
 		Looping = looping;
 	}
 
-	public void release()
+	public synchronized void release()
 	{
-//		GameActivity.Log.debug("CameraPlayer14: release");
-		if (null != mOESTextureRenderer)
-		{
-			while (WaitOnBitmapRender) ;
-			releaseOESTextureRenderer();
-		}
-		if (null != mBitmapRenderer)
-		{
-			while (WaitOnBitmapRender) ;
-			releaseBitmapRenderer();
-		}
-
+		GameActivity.Log.debug("CameraPlayer14: release");
 		if (null != mCamera)
 		{
-			mCamera.stopPreview();
-			mCamera.release();
+			stopPreview();
+			mCamera.close();
 			mCamera = null;
+			mCameraCaptureSession = null;
 		}
+
+		Prepared = false;
+		
+		audioTracks.clear();
+		videoTracks.clear();
+
+		stopBackgroundThread();
+
+		CameraState = CameraStates.RELEASED;
 	}
 
-	public void reset()
+	public synchronized void reset()
 	{
-//		GameActivity.Log.debug("CameraPlayer14: reset");
+		GameActivity.Log.debug("CameraPlayer14: reset");
+
+		release();
+		
+		Completed = false;
+		Looping = false;
+
 		if (null != mOESTextureRenderer)
 		{
 			while (WaitOnBitmapRender) ;
@@ -730,15 +894,37 @@ public class CameraPlayer14
 			while (WaitOnBitmapRender) ;
 			releaseBitmapRenderer();
 		}
+
+		CameraState = CameraStates.PREINIT;
 	}
 
 	// ======================================================================================
 
+	private boolean usingBitmapRendererImageReader = false;
+	private static final int ImageReaderQueueLength = 16;
+	private BitmapRenderer.UpdateParam updateParam = new BitmapRenderer.UpdateParam();
+
+	public void setUsingBitmapRendererImageReader(boolean enabled)
+	{
+		usingBitmapRendererImageReader = enabled;
+	}
+
 	private boolean CreateBitmapRenderer()
 	{
 		releaseBitmapRenderer();
 
-		mBitmapRenderer = new BitmapRenderer(SwizzlePixels, VulkanRenderer);
+		GameActivity.Log.debug("CameraPlayer - CREATE BitmapRenderer, VIDEO SIZE: " + getVideoWidth() + " x " + getVideoHeight() + ", BRIR: " + usingBitmapRendererImageReader + ", Vulkan: " + VulkanRenderer);
+		
+		if (usingBitmapRendererImageReader)
+		{
+			mBitmapRenderer = new BitmapRendererImageReader(SwizzlePixels, VulkanRenderer,
+					getVideoWidth(), getVideoHeight(), ImageReaderQueueLength);
+		}
+		else
+		{
+			mBitmapRenderer = new BitmapRendererLegacy(SwizzlePixels, VulkanRenderer, false);
+		}
+
 		if (!mBitmapRenderer.isValid())
 		{
 			mBitmapRenderer = null;
@@ -748,6 +934,8 @@ public class CameraPlayer14
 		// set this here as the size may have been set before the GL resources were created.
 		mBitmapRenderer.setSize(getVideoWidth(),getVideoHeight());
 
+		mAutoStartPreview = true;
+
 		setVideoEnabled(true);
 		setAudioEnabled(true);
 		return true;
@@ -759,956 +947,195 @@ public class CameraPlayer14
 		{
 			mBitmapRenderer.release();
 			mBitmapRenderer = null;
+		}
+	}
 
-			if (null != mCamera)
-			{
-				try
-				{
-					mCamera.setPreviewTexture(null);
-				}
-				catch (IOException e)
-				{
-				}
-			}
-		}
-	}
-
-	public void initBitmapRenderer()
-	{
-		// if not already allocated.
-		// Create bitmap renderer's gl resources in the renderer thread.
-		if (null == mBitmapRenderer)
-		{
-			if (!CreateBitmapRenderer())
-			{
-				GameActivity.Log.warn("initBitmapRenderer failed to alloc mBitmapRenderer ");
-				reset();
-			  }
-		}
-	}
-
-	public FrameUpdateInfo getVideoLastFrameData()
-	{
-		if (null == mCamera)
-		{
-			return null;
-		}
-
-		initBitmapRenderer();
-		if (null != mBitmapRenderer)
-		{
-			WaitOnBitmapRender = true;
-			FrameUpdateInfo frameInfo = mBitmapRenderer.updateFrameData();
-			WaitOnBitmapRender = false;
-			return frameInfo;
-		}
-		else
-		{
-			return null;
-		}
-	}
-
-	public FrameUpdateInfo getVideoLastFrame(int destTexture)
-	{
-		if (null == mCamera)
-		{
-			return null;
-		}
-
-//		GameActivity.Log.debug("getVideoLastFrame: " + destTexture);
-		initBitmapRenderer();
-		if (null != mBitmapRenderer)
-		{
-			WaitOnBitmapRender = true;
-			FrameUpdateInfo frameInfo = mBitmapRenderer.updateFrameData(destTexture);
-			WaitOnBitmapRender = false;
-			return frameInfo;
-		}
-		else
-		{
-			return null;
-		}
-	}
-
-	/*
-		All this internal surface view does is manage the
-		offscreen bitmap that the media player decoding can
-		render into for eventual extraction to the UE4 buffers.
-	*/
-	class BitmapRenderer
-		implements android.graphics.SurfaceTexture.OnFrameAvailableListener
-	{
-		private java.nio.Buffer mFrameData = null;
-		private int mLastFramePosition = -1;
-		private android.graphics.SurfaceTexture mSurfaceTexture = null;
-		private int mTextureWidth = -1;
-		private int mTextureHeight = -1;
-		private android.view.Surface mSurface = null;
-		private boolean mFrameAvailable = false;
-		private int mTextureID = -1;
-		private int mFBO = -1;
-		private int mBlitVertexShaderID = -1;
-		private int mBlitFragmentShaderID = -1;
-		private float[] mTransformMatrix = new float[16];
-		private boolean mTriangleVerticesDirty = true;
-		private boolean mTextureSizeChanged = true;
-		private boolean mUseOwnContext = true;
-		private boolean mVulkanRenderer = false;
-		private boolean mSwizzlePixels = false;
-
-		private float mScaleRotation00 = 1.0f;
-		private float mScaleRotation01 = 0.0f;
-		private float mScaleRotation10 = 0.0f;
-		private float mScaleRotation11 = 1.0f;
-		private float mUOffset = 0.0f;
-		private float mVOffset = 0.0f;
-
-		private int GL_TEXTURE_EXTERNAL_OES = 0x8D65;
-
-		private EGLDisplay mEglDisplay;
-		private EGLContext mEglContext;
-		private EGLSurface mEglSurface;
-
-		private EGLDisplay mSavedDisplay;
-		private EGLContext mSavedContext;
-		private EGLSurface mSavedSurfaceDraw;
-		private EGLSurface mSavedSurfaceRead;
-
-		private boolean mCreatedEGLDisplay = false;
-
-		public BitmapRenderer(boolean swizzlePixels, boolean isVulkan)
-		{
-			mSwizzlePixels = swizzlePixels;
-			mVulkanRenderer = isVulkan;
-
-			mEglSurface = EGL14.EGL_NO_SURFACE;
-			mEglContext = EGL14.EGL_NO_CONTEXT;
-			mEglDisplay = EGL14.EGL_NO_DISPLAY;
-			mUseOwnContext = true;
-
-			if (mVulkanRenderer)
-			{
-				mSwizzlePixels = true;
-			}
-			else
-			{
-				String RendererString = GLES20.glGetString(GLES20.GL_RENDERER);
-
-				// Do not use shared context if Adreno before 400 or on older Android than Marshmallow
-				if (RendererString.contains("Adreno (TM) "))
-				{
-					java.util.regex.Pattern pattern = java.util.regex.Pattern.compile("Adreno \\(TM\\) (\\d*)");
-					java.util.regex.Matcher matcher = pattern.matcher(RendererString);
-					if (matcher.find() && matcher.groupCount() >= 1)
-					{
-						int AdrenoVersion = Integer.parseInt(matcher.group(1));
-						if (AdrenoVersion < 400 || android.os.Build.VERSION.SDK_INT < 22)
-						{
-							GameActivity.Log.debug("CameraPlayer14: disabled shared GL context on " + RendererString);
-							mUseOwnContext = false;
-						}
-					}
-				}
-			}
-
-			if (mUseOwnContext)
-			{
-				initContext();
-				saveContext();
-				makeCurrent();
-				initSurfaceTexture();
-				restoreContext();
-			}
-			else
-			{
-				initSurfaceTexture();
-			}
-		}
-
-		private void initContext()
-		{
-			mEglDisplay = EGL14.EGL_NO_DISPLAY;
-			EGLContext shareContext = EGL14.EGL_NO_CONTEXT;
-
-			if (!mVulkanRenderer)
-			{
-				mEglDisplay = EGL14.eglGetCurrentDisplay();
-				shareContext = EGL14.eglGetCurrentContext();
-			}
-			else
-			{
-				mEglDisplay = EGL14.eglGetDisplay(EGL14.EGL_DEFAULT_DISPLAY);
-				if (mEglDisplay == EGL14.EGL_NO_DISPLAY)
-				{
-					GameActivity.Log.error("unable to get EGL14 display");
-					return;
-				}
-				int[] version = new int[2];
-				if (!EGL14.eglInitialize(mEglDisplay, version, 0, version, 1))
-				{
-					mEglDisplay = null;
-					GameActivity.Log.error("unable to initialize EGL14 display");
-					return;
-				}				
-				
-				mCreatedEGLDisplay = true;
-			}
-
-			int[] configSpec = new int[]
-			{
-				EGL14.EGL_RENDERABLE_TYPE, EGL14.EGL_OPENGL_ES2_BIT,
-				EGL14.EGL_SURFACE_TYPE, EGL14.EGL_PBUFFER_BIT,
-				EGL14.EGL_NONE
-			};
-			EGLConfig[] configs = new EGLConfig[1];
-			int[] num_config = new int[1];
-			EGL14.eglChooseConfig(mEglDisplay, configSpec, 0, configs, 0, 1, num_config, 0);
-			int[] contextAttribs = new int[]
-			{
-				EGL14.EGL_CONTEXT_CLIENT_VERSION, 2,
-				EGL14.EGL_NONE
-			};
-			mEglContext = EGL14.eglCreateContext(mEglDisplay, configs[0], shareContext, contextAttribs, 0);
-
-			if (EGL14.eglQueryString(mEglDisplay, EGL14.EGL_EXTENSIONS).contains("EGL_KHR_surfaceless_context"))
-			{
-				mEglSurface = EGL14.EGL_NO_SURFACE;
-			}
-			else
-			{
-				int[] pbufferAttribs = new int[]
-				{
-					EGL14.EGL_NONE
-				};
-				mEglSurface = EGL14.eglCreatePbufferSurface(mEglDisplay, configs[0], pbufferAttribs, 0);
-			}
-		}
-
-		private void saveContext()
-		{
-			mSavedDisplay = EGL14.eglGetCurrentDisplay();
-			mSavedContext = EGL14.eglGetCurrentContext();
-			mSavedSurfaceDraw = EGL14.eglGetCurrentSurface(EGL14.EGL_DRAW);
-			mSavedSurfaceRead = EGL14.eglGetCurrentSurface(EGL14.EGL_READ);
-		}
-
-		private void makeCurrent()
-		{
-			EGL14.eglMakeCurrent(mEglDisplay, mEglSurface, mEglSurface, mEglContext);
-		}
-
-		private void restoreContext()
-		{
-			EGL14.eglMakeCurrent(mSavedDisplay, mSavedSurfaceDraw, mSavedSurfaceRead, mSavedContext);
-		}
-
-		private void initSurfaceTexture()
-		{
-			int[] textures = new int[1];
-			GLES20.glGenTextures(1, textures, 0);
-			mTextureID = textures[0];
-			if (mTextureID <= 0)
-			{
-				GameActivity.Log.error("mTextureID <= 0");
-				release();
-				return;
-			}
-			mSurfaceTexture = new android.graphics.SurfaceTexture(mTextureID);
-			mSurfaceTexture.setOnFrameAvailableListener(this);
-			mSurface = new android.view.Surface(mSurfaceTexture);
-
-			int[] glInt = new int[1];
-
-			GLES20.glGenFramebuffers(1,glInt,0);
-			mFBO = glInt[0];
-			if (mFBO <= 0)
-			{
-				GameActivity.Log.error("mFBO <= 0");
-				release();
-				return;
-			}
-
-			// Special shaders for blit of movie texture.
-			mBlitVertexShaderID = createShader(GLES20.GL_VERTEX_SHADER, mBlitVextexShader);
-			if (mBlitVertexShaderID == 0)
-			{
-				GameActivity.Log.error("mBlitVertexShaderID == 0");
-				release();
-				return;
-			}
-			int mBlitFragmentShaderID = createShader(GLES20.GL_FRAGMENT_SHADER,
-				mSwizzlePixels ? mBlitFragmentShaderBGRA : mBlitFragmentShaderRGBA);
-			if (mBlitFragmentShaderID == 0)
-			{
-				GameActivity.Log.error("mBlitFragmentShaderID == 0");
-				release();
-				return;
-			}
-			mProgram = GLES20.glCreateProgram();
-			if (mProgram <= 0)
-			{
-				GameActivity.Log.error("mProgram <= 0");
-				release();
-				return;
-			}
-			GLES20.glAttachShader(mProgram, mBlitVertexShaderID);
-			GLES20.glAttachShader(mProgram, mBlitFragmentShaderID);
-			GLES20.glLinkProgram(mProgram);
-			int[] linkStatus = new int[1];
-			GLES20.glGetProgramiv(mProgram, GLES20.GL_LINK_STATUS, linkStatus, 0);
-			if (linkStatus[0] != GLES20.GL_TRUE)
-			{
-				GameActivity.Log.error("Could not link program: ");
-				GameActivity.Log.error(GLES20.glGetProgramInfoLog(mProgram));
-				GLES20.glDeleteProgram(mProgram);
-				mProgram = 0;
-				release();
-				return;
-			}
-			mPositionAttrib = GLES20.glGetAttribLocation(mProgram, "Position");
-			mTexCoordsAttrib = GLES20.glGetAttribLocation(mProgram, "TexCoords");
-			mTextureUniform = GLES20.glGetUniformLocation(mProgram, "VideoTexture");
-
-			GLES20.glGenBuffers(1,glInt,0);
-			mBlitBuffer = glInt[0];
-			if (mBlitBuffer <= 0)
-			{
-				GameActivity.Log.error("mBlitBuffer <= 0");
-				release();
-				return;
-			}
-
-			// Create blit mesh.
-			mTriangleVertices = java.nio.ByteBuffer.allocateDirect(
-				mTriangleVerticesData.length * FLOAT_SIZE_BYTES)
-					.order(java.nio.ByteOrder.nativeOrder()).asFloatBuffer();
-			mTriangleVerticesDirty = true;
-
-			// Set up GL state
-			if (mUseOwnContext)
-			{
-				GLES20.glDisable(GLES20.GL_BLEND);
-				GLES20.glDisable(GLES20.GL_CULL_FACE);
-				GLES20.glDisable(GLES20.GL_SCISSOR_TEST);
-				GLES20.glDisable(GLES20.GL_STENCIL_TEST);
-				GLES20.glDisable(GLES20.GL_DEPTH_TEST);
-				GLES20.glDisable(GLES20.GL_DITHER);
-				GLES20.glColorMask(true,true,true,true);
-			}
-		}
-		
-		private void UpdateVertexData()
-		{
-			if (!mTriangleVerticesDirty || mBlitBuffer <= 0)
-			{
-				return;
-			}
-
-			// fill it in
-			mTriangleVertices.position(0);
-			mTriangleVertices.put(mTriangleVerticesData).position(0);
-
-			// save VBO state
-			int[] glInt = new int[1];
-			GLES20.glGetIntegerv(GLES20.GL_ARRAY_BUFFER_BINDING, glInt, 0);
-			int previousVBO = glInt[0];
-			
-			GLES20.glBindBuffer(GLES20.GL_ARRAY_BUFFER, mBlitBuffer);
-			GLES20.glBufferData(GLES20.GL_ARRAY_BUFFER,
-				mTriangleVerticesData.length*FLOAT_SIZE_BYTES,
-				mTriangleVertices, GLES20.GL_STATIC_DRAW);
-
-			// restore VBO state
-			GLES20.glBindBuffer(GLES20.GL_ARRAY_BUFFER, previousVBO);
-			
-			mTriangleVerticesDirty = false;
-		}
-
-		public boolean isValid()
-		{
-			return mSurfaceTexture != null;
-		}
-
-		private int createShader(int shaderType, String source)
-		{
-			int shader = GLES20.glCreateShader(shaderType);
-			if (shader != 0)
-			{
-				GLES20.glShaderSource(shader, source);
-				GLES20.glCompileShader(shader);
-				int[] compiled = new int[1];
-				GLES20.glGetShaderiv(shader, GLES20.GL_COMPILE_STATUS, compiled, 0);
-				if (compiled[0] == 0)
-				{
-					GameActivity.Log.error("Could not compile shader " + shaderType + ":");
-					GameActivity.Log.error(GLES20.glGetShaderInfoLog(shader));
-					GLES20.glDeleteShader(shader);
-					shader = 0;
-				}
-			}
-			return shader;
-		}
-
-		public void onFrameAvailable(android.graphics.SurfaceTexture st)
-		{
-			synchronized(this)
-			{
-				mFrameAvailable = true;
-			}
-		}
-
-		public android.graphics.SurfaceTexture getSurfaceTexture()
-		{
-			return mSurfaceTexture;
-		}
-
-		public android.view.Surface getSurface()
-		{
-			return mSurface;
-		}
-
-		public int getExternalTextureId()
-		{
-			return mTextureID;
-		}
-
-		// NOTE: Synchronized with updateFrameData to prevent frame
-		// updates while the surface may need to get reallocated.
-		public void setSize(int width, int height)
-		{
-			synchronized(this)
-			{
-				if (width != mTextureWidth ||
-					height != mTextureHeight)
-				{
-					mTextureWidth = width;
-					mTextureHeight = height;
-					mFrameData = null;
-					mTextureSizeChanged = true;
-
-				}
-			}
-		}
-
-		public boolean resolutionChanged()
-		{
-			boolean changed;
-			synchronized(this)
-			{
-				changed = mTextureSizeChanged;
-				mTextureSizeChanged = false;
-			}
-			return changed;
-		}
-
-		private static final int FLOAT_SIZE_BYTES = 4;
-		private static final int TRIANGLE_VERTICES_DATA_STRIDE_BYTES = 4 * FLOAT_SIZE_BYTES;
-		private static final int TRIANGLE_VERTICES_DATA_POS_OFFSET = 0;
-		private static final int TRIANGLE_VERTICES_DATA_UV_OFFSET = 2;
-		private float[] mTriangleVerticesData = {
-			// X, Y, U, V
-			-1.0f, -1.0f, 0.f, 0.f,
-			1.0f, -1.0f, 1.f, 0.f,
-			-1.0f, 1.0f, 0.f, 1.f,
-			1.0f, 1.0f, 1.f, 1.f,
-			};
-
-		private java.nio.FloatBuffer mTriangleVertices;
-
-		private final String mBlitVextexShader =
-			"attribute vec2 Position;\n" +
-			"attribute vec2 TexCoords;\n" +
-			"varying vec2 TexCoord;\n" +
-			"void main() {\n" +
-			"	TexCoord = TexCoords;\n" +
-			"	gl_Position = vec4(Position, 0.0, 1.0);\n" +
-			"}\n";
-
-		// NOTE: We read the fragment as BGRA so that in the end, when
-		// we glReadPixels out of the FBO, we get them in that order
-		// and avoid having to swizzle the pixels in the CPU.
-		private final String mBlitFragmentShaderBGRA =
-			"#extension GL_OES_EGL_image_external : require\n" +
-			"uniform samplerExternalOES VideoTexture;\n" +
-			"varying highp vec2 TexCoord;\n" +
-			"void main()\n" +
-			"{\n" +
-			"	gl_FragColor = texture2D(VideoTexture, TexCoord).bgra;\n" +
-			"}\n";
-		private final String mBlitFragmentShaderRGBA =
-			"#extension GL_OES_EGL_image_external : require\n" +
-			"uniform samplerExternalOES VideoTexture;\n" +
-			"varying highp vec2 TexCoord;\n" +
-			"void main()\n" +
-			"{\n" +
-			"	gl_FragColor = texture2D(VideoTexture, TexCoord).rgba;\n" +
-			"}\n";
-
-		private int mProgram;
-		private int mPositionAttrib;
-		private int mTexCoordsAttrib;
-		private int mBlitBuffer;
-		private int mTextureUniform;
-
-		public FrameUpdateInfo updateFrameData()
-		{
-			synchronized(this)
-			{
-				if (null == mFrameData && mTextureWidth > 0 && mTextureHeight > 0)
-				{
-					mFrameData = java.nio.ByteBuffer.allocateDirect(mTextureWidth*mTextureHeight*4);
-				}
-				// Copy surface texture to frame data.
-				if (!copyFrameTexture(0, mFrameData))
-				{
-					return null;
-				}
-			}
-
-			FrameUpdateInfo frameUpdateInfo = new FrameUpdateInfo();
-
-			frameUpdateInfo.Buffer = mFrameData;
-			frameUpdateInfo.CurrentPosition = getCurrentPosition();
-			frameUpdateInfo.FrameReady = true;
-			frameUpdateInfo.RegionChanged = false;
-
-			frameUpdateInfo.ScaleRotation00 = mScaleRotation00;
-			frameUpdateInfo.ScaleRotation01 = mScaleRotation01;
-			frameUpdateInfo.ScaleRotation10 = mScaleRotation10;
-			frameUpdateInfo.ScaleRotation11 = mScaleRotation11;
-
-			frameUpdateInfo.UOffset = mUOffset;
-			frameUpdateInfo.VOffset = mVOffset;
-
-			return frameUpdateInfo;
-		}
-
-		public FrameUpdateInfo updateFrameData(int destTexture)
-		{
-			synchronized(this)
-			{
-				// Copy surface texture to destination texture.
-				if (!copyFrameTexture(destTexture, null))
-				{
-					return null;
-				}
-			}
-
-			FrameUpdateInfo frameUpdateInfo = new FrameUpdateInfo();
-
-			frameUpdateInfo.Buffer = null;
-			frameUpdateInfo.CurrentPosition = getCurrentPosition();
-			frameUpdateInfo.FrameReady = true;
-			frameUpdateInfo.RegionChanged = false;
-
-			frameUpdateInfo.ScaleRotation00 = mScaleRotation00;
-			frameUpdateInfo.ScaleRotation01 = mScaleRotation01;
-			frameUpdateInfo.ScaleRotation10 = mScaleRotation10;
-			frameUpdateInfo.ScaleRotation11 = mScaleRotation11;
-
-			frameUpdateInfo.UOffset = mUOffset;
-			frameUpdateInfo.VOffset = mVOffset;
-
-			return frameUpdateInfo;
-		}
-
-		// Copy the surface texture to another texture, or to raw data.
-		// Note: copying to raw data creates a temporary FBO texture.
-		private boolean copyFrameTexture(int destTexture, java.nio.Buffer destData)
-		{
-			if (!mFrameAvailable)
-			{
-				// We only return fresh data when we generate it. At other
-				// time we return nothing to indicate that there was nothing
-				// new to return. The media player deals with this by keeping
-				// the last frame around and using that for rendering.
-				return false;
-			}
-			mFrameAvailable = false;
-			int current_frame_position = getCurrentPosition();
-			mLastFramePosition = current_frame_position;
-			if (null == mSurfaceTexture)
-			{
-				// Can't update if there's no surface to update into.
-				return false;
-			}
-
-			int[] glInt = new int[1];
-			boolean[] glBool = new boolean[1];
-
-			// Either use own context or save states
-			boolean previousBlend=false, previousCullFace=false, previousScissorTest=false, previousStencilTest=false, previousDepthTest=false, previousDither=false;
-			int previousFBO=0, previousVBO=0, previousMinFilter=0, previousMagFilter=0;
-			int[] previousViewport = new int[4];
-			if (mUseOwnContext)
-			{
-				// Received reports of these not being preserved when changing contexts
-				GLES20.glActiveTexture(GLES20.GL_TEXTURE0);
-				GLES20.glGetTexParameteriv(GLES20.GL_TEXTURE_2D, GLES20.GL_TEXTURE_MIN_FILTER, glInt, 0);
-				previousMinFilter = glInt[0];
-				GLES20.glGetTexParameteriv(GLES20.GL_TEXTURE_2D, GLES20.GL_TEXTURE_MAG_FILTER, glInt, 0);
-				previousMagFilter = glInt[0];
-
-				saveContext();
-				makeCurrent();
-			}
-			else
-			{
-				// Clear gl errors as they can creep in from the UE4 renderer.
-				GLES20.glGetError();
-
-				previousBlend = GLES20.glIsEnabled(GLES20.GL_BLEND);
-				previousCullFace = GLES20.glIsEnabled(GLES20.GL_CULL_FACE);
-				previousScissorTest = GLES20.glIsEnabled(GLES20.GL_SCISSOR_TEST);
-				previousStencilTest = GLES20.glIsEnabled(GLES20.GL_STENCIL_TEST);
-				previousDepthTest = GLES20.glIsEnabled(GLES20.GL_DEPTH_TEST);
-				previousDither = GLES20.glIsEnabled(GLES20.GL_DITHER);
-				GLES20.glGetIntegerv(GLES20.GL_FRAMEBUFFER_BINDING, glInt, 0);
-				previousFBO = glInt[0];
-				GLES20.glGetIntegerv(GLES20.GL_ARRAY_BUFFER_BINDING, glInt, 0);
-				previousVBO = glInt[0];
-				GLES20.glGetIntegerv(GLES20.GL_VIEWPORT, previousViewport, 0);
-
-				GLES20.glActiveTexture(GLES20.GL_TEXTURE0);
-				GLES20.glGetTexParameteriv(GLES20.GL_TEXTURE_2D, GLES20.GL_TEXTURE_MIN_FILTER, glInt, 0);
-				previousMinFilter = glInt[0];
-				GLES20.glGetTexParameteriv(GLES20.GL_TEXTURE_2D, GLES20.GL_TEXTURE_MAG_FILTER, glInt, 0);
-				previousMagFilter = glInt[0];
-
-				glVerify("save state");
-			}
-
-			// Get the latest video texture frame.
-			mSurfaceTexture.updateTexImage();
-
-			mSurfaceTexture.getTransformMatrix(mTransformMatrix);
-
-			int degrees = (GameActivity.Get().getCurrentDeviceRotationDegree() + CameraRotationOffset) % 360;
-			if (degrees == 0) {
-				mScaleRotation00 = mTransformMatrix[1];
-				mScaleRotation01 = mTransformMatrix[5];
-				mScaleRotation10 = -mTransformMatrix[0];
-				mScaleRotation11 = -mTransformMatrix[4];
-				mUOffset = mTransformMatrix[13];
-				mVOffset = 1.0f - mTransformMatrix[12];
-				if (CameraId == iBackId)
-				{
-					mScaleRotation00 = -mScaleRotation00;
-					mScaleRotation01 = -mScaleRotation01;
-					mUOffset = 1.0f - mUOffset;
-				}
-			}
-			else if (degrees == 90) {
-				mScaleRotation00 = mTransformMatrix[0];
-				mScaleRotation01 = mTransformMatrix[4];
-				mScaleRotation10 = mTransformMatrix[1];
-				mScaleRotation11 = mTransformMatrix[5];
-				mUOffset = mTransformMatrix[12];
-				mVOffset = mTransformMatrix[13];
-				if (CameraId == iFrontId)
-				{
-					mScaleRotation00 = -mScaleRotation00;
-					mScaleRotation01 = -mScaleRotation01;
-					mUOffset = 1.0f - mUOffset;
-				}
-				mScaleRotation10 = -mScaleRotation10;
-				mScaleRotation11 = -mScaleRotation11;
-				mVOffset = 1.0f - mVOffset;
-			}
-			else if (degrees == 180) {
-				mScaleRotation00 = mTransformMatrix[1];
-				mScaleRotation01 = mTransformMatrix[5];
-				mScaleRotation10 = mTransformMatrix[0];
-				mScaleRotation11 = mTransformMatrix[4];
-				mUOffset = mTransformMatrix[13];
-				mVOffset = mTransformMatrix[12];
-				if (CameraId == iFrontId)
-				{
-					mScaleRotation00 = -mScaleRotation00;
-					mScaleRotation01 = -mScaleRotation01;
-					mUOffset = 1.0f - mUOffset;
-				}
-			}
-			else {
-				mScaleRotation00 = mTransformMatrix[0];
-				mScaleRotation01 = mTransformMatrix[4];
-				mScaleRotation10 = mTransformMatrix[1];
-				mScaleRotation11 = mTransformMatrix[5];
-				mUOffset = mTransformMatrix[12];
-				mVOffset = mTransformMatrix[13];
-				if (CameraId == iBackId)
-				{
-					mScaleRotation00 = -mScaleRotation00;
-					mScaleRotation01 = -mScaleRotation01;
-					mUOffset = 1.0f - mUOffset;
-				}
-			}
-
-			if (null != destData)
+	public void initBitmapRenderer()
+	{
+		// if not already allocated.
+		// Create bitmap renderer's gl resources in the renderer thread.
+		if (null == mBitmapRenderer)
+		{
+			if (!CreateBitmapRenderer())
 			{
-				// Rewind data so that we can write to it.
-				destData.position(0);
+				GameActivity.Log.warn("CameraPlayer14: initBitmapRenderer - failed to alloc mBitmapRenderer");
+				reset();
 			}
+		}
+	}
 
-			if (!mUseOwnContext)
-			{
-				GLES20.glDisable(GLES20.GL_BLEND);
-				GLES20.glDisable(GLES20.GL_CULL_FACE);
-				GLES20.glDisable(GLES20.GL_SCISSOR_TEST);
-				GLES20.glDisable(GLES20.GL_STENCIL_TEST);
-				GLES20.glDisable(GLES20.GL_DEPTH_TEST);
-				GLES20.glDisable(GLES20.GL_DITHER);
-				GLES20.glColorMask(true,true,true,true);
-
-				glVerify("reset state");
-			}
+	private float[] mTransformMatrix = new float[16];
 
-			GLES20.glViewport(0, 0, mTextureWidth, mTextureHeight);
+	private void UpdateUVTransformation(FrameUpdateInfo updateInfo)
+	{
+		if (usingBitmapRendererImageReader)
+		{
+			return;
+		}
 
-			glVerify("set viewport");
+		android.graphics.SurfaceTexture surfaceTexture = mBitmapRenderer.getSurfaceTexture();
+		surfaceTexture.getTransformMatrix(mTransformMatrix);
 
-			// Set-up FBO target texture..
-			int FBOTextureID = 0;
-			if (null != destData)
+		int degrees = (GameActivity.Get().getCurrentDeviceRotationDegree() + CameraRotationOffset) % 360;
+		if (degrees == 0) 
+		{
+			updateInfo.ScaleRotation00 = mTransformMatrix[1];
+			updateInfo.ScaleRotation01 = mTransformMatrix[5];
+			updateInfo.ScaleRotation10 = -mTransformMatrix[0];
+			updateInfo.ScaleRotation11 = -mTransformMatrix[4];
+			updateInfo.UOffset = mTransformMatrix[13];
+			updateInfo.VOffset = 1.0f - mTransformMatrix[12];
+			if (CameraId == CameraBackId)
 			{
-				// Create temporary FBO for data copy.
-				GLES20.glGenTextures(1,glInt,0);
-				FBOTextureID = glInt[0];
+				updateInfo.ScaleRotation00 = -updateInfo.ScaleRotation00;
+				updateInfo.ScaleRotation01 = -updateInfo.ScaleRotation01;
+				updateInfo.UOffset = 1.0f - updateInfo.UOffset;
 			}
-			else
+		}
+		else if (degrees == 90) 
+		{
+			updateInfo.ScaleRotation00 = mTransformMatrix[0];
+			updateInfo.ScaleRotation01 = mTransformMatrix[4];
+			updateInfo.ScaleRotation10 = mTransformMatrix[1];
+			updateInfo.ScaleRotation11 = mTransformMatrix[5];
+			updateInfo.UOffset = mTransformMatrix[12];
+			updateInfo.VOffset = mTransformMatrix[13];
+			if (CameraId == CameraFrontId)
 			{
-				// Use the given texture as the FBO.
-				FBOTextureID = destTexture;
+				updateInfo.ScaleRotation00 = -updateInfo.ScaleRotation00;
+				updateInfo.ScaleRotation01 = -updateInfo.ScaleRotation01;
+				updateInfo.UOffset = 1.0f - updateInfo.UOffset;
 			}
-			// Set the FBO to draw into the texture one-to-one.
-			GLES20.glBindTexture(GLES20.GL_TEXTURE_2D, FBOTextureID);
-			GLES20.glTexParameteri(GLES20.GL_TEXTURE_2D,
-				GLES20.GL_TEXTURE_MIN_FILTER, GLES20.GL_NEAREST);
-			GLES20.glTexParameteri(GLES20.GL_TEXTURE_2D,
-				GLES20.GL_TEXTURE_MAG_FILTER, GLES20.GL_NEAREST);
-			GLES20.glTexParameteri(GLES20.GL_TEXTURE_2D,
-				GLES20.GL_TEXTURE_WRAP_S, GLES20.GL_CLAMP_TO_EDGE);
-			GLES20.glTexParameteri(GLES20.GL_TEXTURE_2D,
-				GLES20.GL_TEXTURE_WRAP_T, GLES20.GL_CLAMP_TO_EDGE);
-			// Create the temp FBO data if needed.
-			if (null != destData)
+			updateInfo.ScaleRotation10 = -updateInfo.ScaleRotation10;
+			updateInfo.ScaleRotation11 = -updateInfo.ScaleRotation11;
+			updateInfo.VOffset = 1.0f - updateInfo.VOffset;
+		}
+		else if (degrees == 180) 
+		{
+			updateInfo.ScaleRotation00 = mTransformMatrix[1];
+			updateInfo.ScaleRotation01 = mTransformMatrix[5];
+			updateInfo.ScaleRotation10 = mTransformMatrix[0];
+			updateInfo.ScaleRotation11 = mTransformMatrix[4];
+			updateInfo.UOffset = mTransformMatrix[13];
+			updateInfo.VOffset = mTransformMatrix[12];
+			if (CameraId == CameraFrontId)
 			{
-				//int w = 1<<(32-Integer.numberOfLeadingZeros(mTextureWidth-1));
-				//int h = 1<<(32-Integer.numberOfLeadingZeros(mTextureHeight-1));
-				GLES20.glTexImage2D(GLES20.GL_TEXTURE_2D, 0,
-					GLES20.GL_RGBA,
-					mTextureWidth, mTextureHeight,
-					0, GLES20.GL_RGBA, GLES20.GL_UNSIGNED_BYTE, null);
+				updateInfo.ScaleRotation00 = -updateInfo.ScaleRotation00;
+				updateInfo.ScaleRotation01 = -updateInfo.ScaleRotation01;
+				updateInfo.UOffset = 1.0f - updateInfo.UOffset;
 			}
-
-			glVerify("set-up FBO texture");
-
-			// Set to render to the FBO.
-			GLES20.glBindFramebuffer(GLES20.GL_FRAMEBUFFER, mFBO);
-
-			GLES20.glFramebufferTexture2D(
-				GLES20.GL_FRAMEBUFFER,
-				GLES20.GL_COLOR_ATTACHMENT0,
-				GLES20.GL_TEXTURE_2D, FBOTextureID, 0);
-
-			// check status
-			int status = GLES20.glCheckFramebufferStatus(GLES20.GL_FRAMEBUFFER);
-			if (status != GLES20.GL_FRAMEBUFFER_COMPLETE)
+		}
+		else 
+		{
+			updateInfo.ScaleRotation00 = mTransformMatrix[0];
+			updateInfo.ScaleRotation01 = mTransformMatrix[4];
+			updateInfo.ScaleRotation10 = mTransformMatrix[1];
+			updateInfo.ScaleRotation11 = mTransformMatrix[5];
+			updateInfo.UOffset = mTransformMatrix[12];
+			updateInfo.VOffset = mTransformMatrix[13];
+			if (CameraId == CameraBackId)
 			{
-				GameActivity.Log.warn("Failed to complete framebuffer attachment ("+status+")");
+				updateInfo.ScaleRotation00 = -updateInfo.ScaleRotation00;
+				updateInfo.ScaleRotation01 = -updateInfo.ScaleRotation01;
+				updateInfo.UOffset = 1.0f - updateInfo.UOffset;
 			}
+		}
+	}
 
-			// The special shaders to render from the video texture.
-			GLES20.glUseProgram(mProgram);
-
-			// Set the mesh that renders the video texture.
-			UpdateVertexData();
-			GLES20.glBindBuffer(GLES20.GL_ARRAY_BUFFER, mBlitBuffer);
-			GLES20.glEnableVertexAttribArray(mPositionAttrib);
-			GLES20.glVertexAttribPointer(mPositionAttrib, 2, GLES20.GL_FLOAT, false,
-				TRIANGLE_VERTICES_DATA_STRIDE_BYTES, 0);
-			GLES20.glEnableVertexAttribArray(mTexCoordsAttrib);
-			GLES20.glVertexAttribPointer(mTexCoordsAttrib, 2, GLES20.GL_FLOAT, false,
-				TRIANGLE_VERTICES_DATA_STRIDE_BYTES,
-				TRIANGLE_VERTICES_DATA_UV_OFFSET*FLOAT_SIZE_BYTES);
-
-			glVerify("setup movie texture read");
-
-			GLES20.glClear( GLES20.GL_COLOR_BUFFER_BIT);
-
-			// connect 'VideoTexture' to video source texture (mTextureID).
-			// mTextureID is bound to GL_TEXTURE_EXTERNAL_OES in updateTexImage
-			GLES20.glUniform1i(mTextureUniform, 0);
-			GLES20.glActiveTexture(GLES20.GL_TEXTURE0);
-			GLES20.glBindTexture(GL_TEXTURE_EXTERNAL_OES, mTextureID);
-
-			// Draw the video texture mesh.
-			GLES20.glDrawArrays(GLES20.GL_TRIANGLE_STRIP, 0, 4);
-
-			GLES20.glFlush();
-
-			// Read the FBO texture pixels into raw data.
-			if (null != destData)
-			{
-				GLES20.glReadPixels(
-					0, 0, mTextureWidth, mTextureHeight,
-					GLES20.GL_RGBA, GLES20.GL_UNSIGNED_BYTE,
-					destData);
-			}
+	public FrameUpdateInfo getVideoLastFrameData()
+	{
+		if (usingBitmapRendererImageReader)
+		{
+			return null;
+		}
 
-			glVerify("draw & read movie texture");
+		if (null == mCamera)
+		{
+			return null;
+		}
 
-			// Restore state and cleanup.
-			if (mUseOwnContext)
+		initBitmapRenderer();
+		if (null != mBitmapRenderer)
+		{
+			WaitOnBitmapRender = true;
+			updateParam.destTexture = 0;
+			BitmapRenderer.FrameData data = mBitmapRenderer.updateFrameData(updateParam);
+			WaitOnBitmapRender = false;
+			if (data != null)
 			{
-				GLES20.glFramebufferTexture2D(
-					GLES20.GL_FRAMEBUFFER,
-					GLES20.GL_COLOR_ATTACHMENT0,
-					GLES20.GL_TEXTURE_2D, 0, 0);
-
-				if (null != destData && FBOTextureID > 0)
-				{
-					glInt[0] = FBOTextureID;
-					GLES20.glDeleteTextures(1, glInt, 0);
-				}
+				FrameUpdateInfo frameUpdateInfo = new FrameUpdateInfo();
+				frameUpdateInfo.Buffer = data.binaryBuffer;
+				frameUpdateInfo.BRFrameData = null;
+				frameUpdateInfo.FrameReady = true;
+				frameUpdateInfo.RegionChanged = false;
 
-				restoreContext();
-
-				// Restore previous texture filtering
-				GLES20.glTexParameteri(GLES20.GL_TEXTURE_2D, GLES20.GL_TEXTURE_MIN_FILTER, previousMinFilter);
-				GLES20.glTexParameteri(GLES20.GL_TEXTURE_2D, GLES20.GL_TEXTURE_MAG_FILTER, previousMagFilter);
-			}
-			else
-			{
-				GLES20.glBindFramebuffer(GLES20.GL_FRAMEBUFFER, previousFBO);
-				if (null != destData && FBOTextureID > 0)
-				{
-					glInt[0] = FBOTextureID;
-					GLES20.glDeleteTextures(1, glInt, 0);
-				}
-				GLES20.glBindBuffer(GLES20.GL_ARRAY_BUFFER, previousVBO);
-
-				GLES20.glViewport(previousViewport[0], previousViewport[1],	previousViewport[2], previousViewport[3]);
-				if (previousBlend) GLES20.glEnable(GLES20.GL_BLEND);
-				if (previousCullFace) GLES20.glEnable(GLES20.GL_CULL_FACE);
-				if (previousScissorTest) GLES20.glEnable(GLES20.GL_SCISSOR_TEST);
-				if (previousStencilTest) GLES20.glEnable(GLES20.GL_STENCIL_TEST);
-				if (previousDepthTest) GLES20.glEnable(GLES20.GL_DEPTH_TEST);
-				if (previousDither) GLES20.glEnable(GLES20.GL_DITHER);
-
-				GLES20.glTexParameteri(GLES20.GL_TEXTURE_2D, GLES20.GL_TEXTURE_MIN_FILTER, previousMinFilter);
-				GLES20.glTexParameteri(GLES20.GL_TEXTURE_2D, GLES20.GL_TEXTURE_MAG_FILTER, previousMagFilter);
-
-				// invalidate cached state in RHI
-				GLES20.glDisableVertexAttribArray(mPositionAttrib);
-				GLES20.glDisableVertexAttribArray(mTexCoordsAttrib);
-				nativeClearCachedAttributeState(mPositionAttrib, mTexCoordsAttrib);
+				UpdateUVTransformation(frameUpdateInfo);
+				
+				return frameUpdateInfo;
 			}
-
-			return true;
 		}
+		return null;
+	}
 
-		private void showGlError(String op, int error)
+	public FrameUpdateInfo getVideoLastFrame(int destTexture)
+	{
+		if (usingBitmapRendererImageReader)
 		{
-			switch (error)
-			{
-				case GLES20.GL_INVALID_ENUM:						GameActivity.Log.error("CameraPlayer$BitmapRenderer: " + op + ": glGetError GL_INVALID_ENUM");  break;
-				case GLES20.GL_INVALID_OPERATION:					GameActivity.Log.error("CameraPlayer$BitmapRenderer: " + op + ": glGetError GL_INVALID_OPERATION");  break;
-				case GLES20.GL_INVALID_FRAMEBUFFER_OPERATION:		GameActivity.Log.error("CameraPlayer$BitmapRenderer: " + op + ": glGetError GL_INVALID_FRAMEBUFFER_OPERATION");  break;
-				case GLES20.GL_INVALID_VALUE:						GameActivity.Log.error("CameraPlayer$BitmapRenderer: " + op + ": glGetError GL_INVALID_VALUE");  break;
-				case GLES20.GL_FRAMEBUFFER_INCOMPLETE_ATTACHMENT:	GameActivity.Log.error("CameraPlayer$BitmapRenderer: " + op + ": glGetError GL_FRAMEBUFFER_INCOMPLETE_ATTACHMENT");  break;
-				case GLES20.GL_FRAMEBUFFER_INCOMPLETE_DIMENSIONS:	GameActivity.Log.error("CameraPlayer$BitmapRenderer: " + op + ": glGetError GL_FRAMEBUFFER_INCOMPLETE_DIMENSIONS");  break;
-				case GLES20.GL_FRAMEBUFFER_UNSUPPORTED:				GameActivity.Log.error("CameraPlayer$BitmapRenderer: " + op + ": glGetError GL_FRAMEBUFFER_UNSUPPORTED");  break;
-				case GLES20.GL_OUT_OF_MEMORY:						GameActivity.Log.error("CameraPlayer$BitmapRenderer: " + op + ": glGetError GL_OUT_OF_MEMORY");  break;
-				default:											GameActivity.Log.error("CameraPlayer$BitmapRenderer: " + op + ": glGetError " + error);
-			}
+			return null;
 		}
 
-		private void glVerify(String op)
+		if (null == mCamera)
 		{
-			int error;
-			while ((error = GLES20.glGetError()) != GLES20.GL_NO_ERROR)
-			{
-				showGlError(op, error);
-				throw new RuntimeException(op + ": glGetError " + error);
-			}
+			return null;
 		}
 
-		private void glWarn(String op)
+//		GameActivity.Log.debug("getVideoLastFrame: " + destTexture);
+		initBitmapRenderer();
+		if (null != mBitmapRenderer)
 		{
-			int error;
-			while ((error = GLES20.glGetError()) != GLES20.GL_NO_ERROR)
+			WaitOnBitmapRender = true;
+			updateParam.destTexture = destTexture;
+			BitmapRenderer.FrameData data = mBitmapRenderer.updateFrameData(updateParam);
+			WaitOnBitmapRender = false;
+			if (data != null)
 			{
-				showGlError(op, error);
+				FrameUpdateInfo frameUpdateInfo = new FrameUpdateInfo();
+				frameUpdateInfo.BRFrameData = null;
+				frameUpdateInfo.Buffer = null;
+				frameUpdateInfo.FrameReady = true;
+				frameUpdateInfo.RegionChanged = false;
+				
+				UpdateUVTransformation(frameUpdateInfo);
+
+				return frameUpdateInfo;
 			}
 		}
+		return null;
+	}
 
-		public void release()
+	public FrameUpdateInfo getVideoLastMediaFrameData()
+	{
+		if (usingBitmapRendererImageReader)
 		{
-			if (null != mSurface)
-			{
-				mSurface.release();
-				mSurface = null;
-			}
-			if (null != mSurfaceTexture)
-			{
-				mSurfaceTexture.release();
-				mSurfaceTexture = null;
-			}
-			int[] glInt = new int[1];
-			if (mBlitBuffer > 0)
-			{
-				glInt[0] = mBlitBuffer;
-				GLES20.glDeleteBuffers(1,glInt,0);
-				mBlitBuffer = -1;
-			}
-			if (mProgram > 0)
-			{
-				GLES20.glDeleteProgram(mProgram);
-				mProgram = -1;
-			}
-			if (mBlitVertexShaderID > 0)
-			{
-				GLES20.glDeleteShader(mBlitVertexShaderID);
-				mBlitVertexShaderID = -1;
-			}
-			if (mBlitFragmentShaderID > 0)
-			{
-				GLES20.glDeleteShader(mBlitFragmentShaderID);
-				mBlitFragmentShaderID = -1;
-			}
-			if (mFBO > 0)
-			{
-				glInt[0] = mFBO;
-				GLES20.glDeleteFramebuffers(1,glInt,0);
-				mFBO = -1;
-			}
-			if (mTextureID > 0)
-			{
-				glInt[0] = mTextureID;
-				GLES20.glDeleteTextures(1,glInt,0);
-				mTextureID = -1;
-			}
-			if (mEglSurface != EGL14.EGL_NO_SURFACE)
-			{
-				EGL14.eglDestroySurface(mEglDisplay, mEglSurface);
-				mEglSurface = EGL14.EGL_NO_SURFACE;
-			}
-			if (mEglContext != EGL14.EGL_NO_CONTEXT)
+			initBitmapRenderer();
+			if (mBitmapRenderer != null)
 			{
-				EGL14.eglDestroyContext(mEglDisplay, mEglContext);
-				mEglContext = EGL14.EGL_NO_CONTEXT;
-			}
-			if (mCreatedEGLDisplay)
-			{
-				EGL14.eglTerminate(mEglDisplay);
-				mEglDisplay = EGL14.EGL_NO_DISPLAY;
-				mCreatedEGLDisplay = false;	
+				WaitOnBitmapRender = true;
+				BitmapRenderer.FrameData data = mBitmapRenderer.updateFrameData(updateParam);
+				WaitOnBitmapRender = false;
+				if (data != null)
+				{
+					FrameUpdateInfo frameUpdateInfo = new FrameUpdateInfo();
+					frameUpdateInfo.BRFrameData = data;
+					frameUpdateInfo.Buffer = null;
+					frameUpdateInfo.FrameReady = true;
+					frameUpdateInfo.RegionChanged = false;
+					return frameUpdateInfo;
+				}
 			}
 		}
-	};
-
-	public native void nativeClearCachedAttributeState(int PositionAttrib, int TexCoordsAttrib);
+		return null;
+	}
 
 	// ======================================================================================
 
@@ -1716,6 +1143,8 @@ public class CameraPlayer14
 	{
 		releaseOESTextureRenderer();
 
+		GameActivity.Log.debug("CameraPlayer - CREATE OESTextureRenderer, TextureId: " + OESTextureId);
+		
 		mOESTextureRenderer = new OESTextureRenderer(OESTextureId);
 		if (!mOESTextureRenderer.isValid())
 		{
@@ -1726,6 +1155,8 @@ public class CameraPlayer14
 		// set this here as the size may have been set before the GL resources were created.
 		mOESTextureRenderer.setSize(getVideoWidth(),getVideoHeight());
 
+		mAutoStartPreview = true;
+
 		setVideoEnabled(true);
 		setAudioEnabled(true);
 		return true;
@@ -1737,17 +1168,6 @@ public class CameraPlayer14
 		{
 			mOESTextureRenderer.release();
 			mOESTextureRenderer = null;
-
-			if (null != mCamera)
-			{
-				try
-				{
-					mCamera.setPreviewTexture(null);
-				}
-				catch (IOException e)
-				{
-				}
-			}
 		}
 	}
 	
@@ -1757,7 +1177,7 @@ public class CameraPlayer14
 		{
 			if (!CreateOESTextureRenderer(externalTextureId))
 			{
-				GameActivity.Log.warn("updateVideoFrame failed to alloc mOESTextureRenderer ");
+				GameActivity.Log.warn("CameraPlayer14: updateVideoFrame - failed to alloc mOESTextureRenderer");
 				reset();
 				return null;
 			}
@@ -1931,7 +1351,7 @@ public class CameraPlayer14
 				mScaleRotation11 = -mTransformMatrix[4];
 				mUOffset = mTransformMatrix[13];
 				mVOffset = 1.0f - mTransformMatrix[12];
-				if (CameraId == iBackId)
+				if (CameraId == CameraBackId)
 				{
 					mScaleRotation00 = -mScaleRotation00;
 					mScaleRotation01 = -mScaleRotation01;
@@ -1945,7 +1365,7 @@ public class CameraPlayer14
 				mScaleRotation11 = mTransformMatrix[5];
 				mUOffset = mTransformMatrix[12];
 				mVOffset = mTransformMatrix[13];
-				if (CameraId == iFrontId)
+				if (CameraId == CameraFrontId)
 				{
 					mScaleRotation00 = -mScaleRotation00;
 					mScaleRotation01 = -mScaleRotation01;
@@ -1962,7 +1382,7 @@ public class CameraPlayer14
 				mScaleRotation11 = mTransformMatrix[4];
 				mUOffset = mTransformMatrix[13];
 				mVOffset = mTransformMatrix[12];
-				if (CameraId == iFrontId)
+				if (CameraId == CameraFrontId)
 				{
 					mScaleRotation00 = -mScaleRotation00;
 					mScaleRotation01 = -mScaleRotation01;
@@ -1976,7 +1396,7 @@ public class CameraPlayer14
 				mScaleRotation11 = mTransformMatrix[5];
 				mUOffset = mTransformMatrix[12];
 				mVOffset = mTransformMatrix[13];
-				if (CameraId == iBackId)
+				if (CameraId == CameraBackId)
 				{
 					mScaleRotation00 = -mScaleRotation00;
 					mScaleRotation01 = -mScaleRotation01;
@@ -2050,7 +1470,7 @@ public class CameraPlayer14
 				VideoTracks[TrackIndex].FrameRate = track.FrameRate;
 				VideoTracks[TrackIndex].FrameRateLow = track.FrameRateLow;
 				VideoTracks[TrackIndex].FrameRateHigh = track.FrameRateHigh;
-//				GameActivity.Log.debug(TrackIndex + " (" + Index + ") Video: Mime=" + VideoTracks[TrackIndex].MimeType + ", Lang=" + VideoTracks[TrackIndex].Language + ", Width=" + VideoTracks[TrackIndex].Width + ", Height=" + VideoTracks[TrackIndex].Height + ", FrameRate=" + VideoTracks[TrackIndex].FrameRate);
+				//GameActivity.Log.debug(TrackIndex + "Video: Mime=" + VideoTracks[TrackIndex].MimeType + ", Lang=" + VideoTracks[TrackIndex].Language + ", Width=" + VideoTracks[TrackIndex].Width + ", Height=" + VideoTracks[TrackIndex].Height + ", FrameRate=" + VideoTracks[TrackIndex].FrameRate);
 				TrackIndex++;
 			}
 
@@ -2075,72 +1495,169 @@ public class CameraPlayer14
 		}
 
 		VideoTrackInfo[] VideoTracks = new VideoTrackInfo[0];
+		GameActivity.Log.error("CameraPlayer14: Empty video tracks.");
 
 		return VideoTracks;
 	}
 
 	// ======================================================================================
 
+	private static final SparseIntArray ORIENTATIONS = new SparseIntArray();
+
+	static {
+		ORIENTATIONS.append(Surface.ROTATION_0, 0);
+		ORIENTATIONS.append(Surface.ROTATION_90, 90);
+		ORIENTATIONS.append(Surface.ROTATION_180, 180);
+		ORIENTATIONS.append(Surface.ROTATION_270, 270);
+	}
+	
+	private int sensorToDeviceRotation(String cameraId, int deviceOrientation)
+	{
+		try
+		{
+			CameraCharacteristics cameraChar = mCameraManager.getCameraCharacteristics(cameraId);
+			int sensorOrientation = cameraChar.get(CameraCharacteristics.SENSOR_ORIENTATION);
+			// Get device orientation in degrees
+			deviceOrientation = ORIENTATIONS.get(deviceOrientation);
+			// Reverse device orientation for front-facing cameras
+			if (cameraChar.get(CameraCharacteristics.LENS_FACING) == CameraCharacteristics.LENS_FACING_FRONT)
+			{
+				deviceOrientation = -deviceOrientation;
+			}
+			// Calculate desired JPEG orientation relative to camera orientation to make
+			// the image upright relative to the device orientation
+			return (sensorOrientation + deviceOrientation + 360) % 360;
+		}
+		catch (CameraAccessException e)
+		{
+			throw new RuntimeException(e);
+		}
+	}
+
 	public boolean takePicture(String Filename, int width, int height)
 	{
-		if (null == mCamera || CameraState != CameraStates.START)
+		synchronized (this)
 		{
-			return false;
+			if (null == mCamera || CameraState != CameraStates.START)
+			{
+				GameActivity.Log.error("CameraPlayer14: takePicture - Camera not ready!");
+				return false;
+			}
+		}
+
+		if (width <= 0 || height <= 0)
+		{
+			width = CameraWidth;
+			height = CameraHeight;
+			GameActivity.Log.warn("CameraPlayer14: takePicture - invalid width or height, using current camera size: " + width + "x" + height);
 		}
 
-		if (width != 0 && height != 0)
+		try
 		{
 			// find best matched picture size
-			Camera.Parameters param = mCamera.getParameters();
-			List<Camera.Size> pictureSizes = param.getSupportedPictureSizes();
-			if (pictureSizes.size() > 0)
+			int bestWidth = -1;
+			int bestHeight = -1;
+			CameraCharacteristics cameraChar = mCameraManager.getCameraCharacteristics(CameraId);
+			StreamConfigurationMap configs = cameraChar.get(CameraCharacteristics.SCALER_STREAM_CONFIGURATION_MAP);
+			android.util.Size[] pictureSizes = configs.getOutputSizes(ImageFormat.JPEG);
+			if (pictureSizes.length > 0)
 			{
-				int bestWidth = -1;
-				int bestHeight = -1;
 				int bestScore = 0;
-				for (int index = 0; index < pictureSizes.size(); ++index)
+				for (int index = 0; index < pictureSizes.length; ++index)
 				{
-					Camera.Size camSize = pictureSizes.get(index);
-					GameActivity.Log.debug("takePicture: pictureSize " + index + ": " + camSize.width + " x " + camSize.height);
-					int diffWidth = camSize.width - width;
-					int diffHeight = camSize.height - height;
+					android.util.Size camSize = pictureSizes[index];
+					int diffWidth = camSize.getWidth() - width;
+					int diffHeight = camSize.getHeight() - height;
 					int score = diffWidth * diffWidth + diffHeight * diffHeight;
 					if (bestWidth == -1 || score < bestScore)
 					{
-						bestWidth = camSize.width;
-						bestHeight = camSize.height;
+						bestWidth = camSize.getWidth();
+						bestHeight = camSize.getHeight();
 						bestScore = score;
 					}
 				}
-				GameActivity.Log.debug("takePicture: using " + bestWidth + " x " + bestHeight);
-				param.setPictureSize(bestWidth, bestHeight);
-				mCamera.setParameters(param);
+				GameActivity.Log.debug("CameraPlayer14: takePicture - final size: " + bestWidth + " x " + bestHeight);
 			}
-		}
-		
-		final String OutFilename = Filename;
-		mCamera.takePicture(null, null, new Camera.PictureCallback()
+			else
+			{
+				bestWidth = width;
+				bestHeight = height;
+			}
+
+			stopPreview();
+
+			final ImageReader.OnImageAvailableListener readerListener = reader -> {
+				Image image = reader.acquireNextImage();
+				ByteBuffer buffer = image.getPlanes()[0].getBuffer();
+				byte[] bytes = new byte[buffer.capacity()];
+				buffer.get(bytes);
+				try
+				{
+					FileOutputStream outStream = new FileOutputStream(Filename);
+					outStream.write(bytes);
+					outStream.close();
+				}
+				catch (IOException e)
+				{
+					throw new RuntimeException(e);
+				}
+				finally
+				{
+					image.close();
+				}
+			};
+
+			ImageReader reader = ImageReader.newInstance(bestWidth, bestHeight, ImageFormat.JPEG, 1);
+			reader.setOnImageAvailableListener(readerListener, mBackgroundHandler);
+
+			CaptureRequest.Builder captureRequestBuilder = mCamera.createCaptureRequest(CameraDevice.TEMPLATE_STILL_CAPTURE);
+			captureRequestBuilder.addTarget(reader.getSurface());
+			captureRequestBuilder.set(CaptureRequest.CONTROL_MODE, CaptureRequest.CONTROL_MODE_AUTO);
+			int rotation = GameActivity.Get().getDisplay().getRotation();
+			int totalRotation = sensorToDeviceRotation(CameraId, rotation);
+			captureRequestBuilder.set(CaptureRequest.JPEG_ORIENTATION, totalRotation);			
+
+			final CameraCaptureSession.CaptureCallback captureListener = new CameraCaptureSession.CaptureCallback()
 			{
 				@Override
-				public void onPictureTaken(byte[] data, Camera camera) {
-					FileOutputStream outStream = null;
-					try {
-						outStream = new FileOutputStream(OutFilename);
-						outStream.write(data);
-						outStream.close();
+				public void onCaptureCompleted(CameraCaptureSession session, CaptureRequest request, TotalCaptureResult result)
+				{
+					super.onCaptureCompleted(session, request, result);
+					synchronized (this)
+					{
+						CameraState = CameraStates.PAUSE;
+						mCameraCaptureSession = null;
+						setVideoEnabled(true);
 					}
-					catch (FileNotFoundException e) {
-						e.printStackTrace();
+				}
+			};
+
+			List<Surface> outputSurfaces = new ArrayList<Surface>(1);
+			outputSurfaces.add(reader.getSurface());
+			mCamera.createCaptureSession(outputSurfaces, new CameraCaptureSession.StateCallback()
+			{
+				@Override
+				public void onConfigured(CameraCaptureSession session)
+				{
+					try
+					{
+						session.capture(captureRequestBuilder.build(), captureListener, mBackgroundHandler);
 					}
-					catch (IOException e) {
+					catch (CameraAccessException e)
+					{
 						e.printStackTrace();
 					}
-					finally {
-					}
 				}
-			});
-		CameraState = CameraStates.PAUSE;
-		
+
+				@Override
+				public void onConfigureFailed(CameraCaptureSession session) {}
+			}, mBackgroundHandler);
+		}
+		catch (CameraAccessException e)
+		{
+			throw new RuntimeException(e);
+		}
+
 		return true;
 	}
 }
diff --git a/Engine/Plugins/Media/AndroidCamera/Source/AndroidCamera/Private/AndroidJavaCameraPlayer.cpp b/Engine/Plugins/Media/AndroidCamera/Source/AndroidCamera/Private/AndroidJavaCameraPlayer.cpp
--- a/Engine/Plugins/Media/AndroidCamera/Source/AndroidCamera/Private/AndroidJavaCameraPlayer.cpp
+++ b/Engine/Plugins/Media/AndroidCamera/Source/AndroidCamera/Private/AndroidJavaCameraPlayer.cpp
@@ -22,6 +22,13 @@ static jfieldID FindField(JNIEnv* JEnv, jclass Class, const ANSICHAR* FieldName,
 	return Field;
 }
 
+bool AndroidEnableCameraPlayerBRIR = true;
+static FAutoConsoleVariableRef CVarAndroidEnableCameraPlayerBRIR(
+	TEXT("Android.EnableCameraPlayerBRIR"),
+	AndroidEnableCameraPlayerBRIR,
+	TEXT("Determines how the Android Camera Player handles video frame data. True for using ImageReader/HardwareBuffer. False for copying data through the CPU memory. (Default: true)"),
+	ECVF_Default);
+
 FJavaAndroidCameraPlayer::FJavaAndroidCameraPlayer(bool swizzlePixels, bool vulkanRenderer)
 	: FJavaClassObject(GetClassName(), "(ZZ)V", swizzlePixels, vulkanRenderer)
 	, GetDurationMethod(GetClassMethod("getDuration", "()I"))
@@ -33,10 +40,9 @@ FJavaAndroidCameraPlayer::FJavaAndroidCameraPlayer(bool swizzlePixels, bool vulk
 	, IsPreparedMethod(GetClassMethod("isPrepared", "()Z"))
 	, SetDataSourceURLMethod(GetClassMethod("setDataSourceURL", "(Ljava/lang/String;)Z"))
 	, GetDataSourceURLMethod(GetClassMethod("getDataSourceURL", "()Ljava/lang/String;"))
-	, PrepareMethod(GetClassMethod("prepare", "()V"))
-	, PrepareAsyncMethod(GetClassMethod("prepareAsync", "()V"))
 	, SeekToMethod(GetClassMethod("seekTo", "(I)V"))
 	, SetLoopingMethod(GetClassMethod("setLooping", "(Z)V"))
+	, OpenCameraMethod(GetClassMethod("openCamera", "()Z"))
 	, ReleaseMethod(GetClassMethod("release", "()V"))
 	, GetVideoHeightMethod(GetClassMethod("getVideoHeight", "()I"))
 	, GetVideoWidthMethod(GetClassMethod("getVideoWidth", "()I"))
@@ -55,6 +61,8 @@ FJavaAndroidCameraPlayer::FJavaAndroidCameraPlayer(bool swizzlePixels, bool vulk
 	, GetExternalTextureIdMethod(GetClassMethod("getExternalTextureId", "()I"))
 	, UpdateVideoFrameMethod(GetClassMethod("updateVideoFrame", "(I)Lcom/epicgames/unreal/CameraPlayer14$FrameUpdateInfo;"))
 	, TakePictureMethod(GetClassMethod("takePicture", "(Ljava/lang/String;II)Z"))
+	, SetUsingBitmapRendererImageReaderMethod(GetClassMethod("setUsingBitmapRendererImageReader", "(Z)V"))
+	, GetVideoLastMediaFrameDataMethod(GetClassMethod("getVideoLastMediaFrameData", "()Lcom/epicgames/unreal/CameraPlayer14$FrameUpdateInfo;"))
 {
 	VideoTexture = nullptr;
 	bVideoTextureValid = false;
@@ -77,6 +85,9 @@ FJavaAndroidCameraPlayer::FJavaAndroidCameraPlayer(bool swizzlePixels, bool vulk
 		SelectTrackMethod = GetClassMethod("selectTrack", "(I)V");
 	}
 
+	BRSelection = BitmapRendererSelection::Unknown;
+	UsingVulkanRenderer = vulkanRenderer && FAndroidMisc::ShouldUseVulkan();
+
 	JNIEnv* JEnv = FAndroidApplication::GetJavaEnv();
 
 	// get field IDs for FrameUpdateInfo class members
@@ -91,6 +102,7 @@ FJavaAndroidCameraPlayer::FJavaAndroidCameraPlayer(bool swizzlePixels, bool vulk
 	FrameUpdateInfo_ScaleRotation11 = FindField(JEnv, FrameUpdateInfoClass, "ScaleRotation11", "F", false);
 	FrameUpdateInfo_UOffset = FindField(JEnv, FrameUpdateInfoClass, "UOffset", "F", false);
 	FrameUpdateInfo_VOffset = FindField(JEnv, FrameUpdateInfoClass, "VOffset", "F", false);
+	FrameUpdateInfo_BRFrameData = FindField(JEnv, FrameUpdateInfoClass, "BRFrameData", "Lcom/epicgames/unreal/BitmapRenderer$FrameData;", false);
 
 	// get field IDs for AudioTrackInfo class members
 	AudioTrackInfoClass = FAndroidApplication::FindJavaClassGlobalRef("com/epicgames/unreal/CameraPlayer14$AudioTrackInfo");
@@ -249,6 +261,11 @@ void FJavaAndroidCameraPlayer::SetLooping(bool Looping)
 	CallMethod<void>(SetLoopingMethod, Looping);
 }
 
+bool FJavaAndroidCameraPlayer::RestoreCamera()
+{
+	return CallMethod<bool>(OpenCameraMethod);
+}
+
 void FJavaAndroidCameraPlayer::Release()
 {
 	CallMethod<void>(ReleaseMethod);
@@ -279,8 +296,41 @@ void FJavaAndroidCameraPlayer::SetAudioEnabled(bool enabled /*= true*/)
 	CallMethod<void>(SetAudioEnabledMethod, enabled);
 }
 
+void FJavaAndroidCameraPlayer::UpdateScaleRotationOffset()
+{
+	static TMap<EDeviceScreenOrientation, const FVector2f/*X->Sin, Y->Cos*/> OrientationFactorsMap;
+	if (OrientationFactorsMap.IsEmpty())
+	{
+		OrientationFactorsMap.Add(EDeviceScreenOrientation::Portrait, FVector2f(-1.0f, 0.0f));
+		OrientationFactorsMap.Add(EDeviceScreenOrientation::PortraitUpsideDown, FVector2f(1.0f, 0.0f));
+		OrientationFactorsMap.Add(EDeviceScreenOrientation::LandscapeRight, FVector2f(0.0f, -1.0f));
+		OrientationFactorsMap.Add(EDeviceScreenOrientation::LandscapeLeft, FVector2f(0.0f, 1.0f));
+	}
+
+	EDeviceScreenOrientation DeviceOrientation = FAndroidMisc::GetDeviceOrientation();
+	const FVector2f* SinCos = OrientationFactorsMap.Find(DeviceOrientation);
+	if (SinCos != nullptr)
+	{
+		ScaleRotation.Set(SinCos->Y, -SinCos->X, SinCos->X, SinCos->Y);
+		FVector2f TranslationOffset(-0.5f, -0.5f);
+		Offset.X = TranslationOffset.Dot(FVector2f(ScaleRotation.X, ScaleRotation.Y)) - TranslationOffset.X;
+		Offset.Y = TranslationOffset.Dot(FVector2f(ScaleRotation.Z, ScaleRotation.W)) - TranslationOffset.Y;
+		Offset.Z = Offset.W = 0.0f;
+	}
+	else
+	{
+		ScaleRotation.Set(1.0f, 0.0f, 0.0f, 1.0f);
+		Offset.Set(0.0f, 0.0f, 0.0f, 0.0f);
+	}
+}
+
 bool FJavaAndroidCameraPlayer::GetVideoLastFrameData(void* & outPixels, int64 & outCount, int32 *CurrentPosition, bool *bRegionChanged)
 {
+	if (IsUsingBitmapRendererImageReader())
+	{
+		return false;
+	}
+
 	// This can return an exception in some cases
 	JNIEnv*	JEnv = FAndroidApplication::GetJavaEnv();
 	auto Result = NewScopedJavaObject(JEnv, JEnv->CallObjectMethod(Object, GetVideoLastFrameDataMethod.Method));
@@ -304,12 +354,8 @@ bool FJavaAndroidCameraPlayer::GetVideoLastFrameData(void* & outPixels, int64 &
 		*CurrentPosition = (int32)JEnv->GetIntField(*Result, FrameUpdateInfo_CurrentPosition);
 		bool bFrameReady = (bool)JEnv->GetBooleanField(*Result, FrameUpdateInfo_FrameReady);
 		*bRegionChanged = (bool)JEnv->GetBooleanField(*Result, FrameUpdateInfo_RegionChanged);
-		ScaleRotation.X = (float)(float)JEnv->GetFloatField(*Result, FrameUpdateInfo_ScaleRotation00);
-		ScaleRotation.Y = (float)(float)JEnv->GetFloatField(*Result, FrameUpdateInfo_ScaleRotation01);
-		ScaleRotation.Z = (float)(float)JEnv->GetFloatField(*Result, FrameUpdateInfo_ScaleRotation10);
-		ScaleRotation.W = (float)(float)JEnv->GetFloatField(*Result, FrameUpdateInfo_ScaleRotation11);
-		Offset.X = (float)JEnv->GetFloatField(*Result, FrameUpdateInfo_UOffset);
-		Offset.Y = (float)JEnv->GetFloatField(*Result, FrameUpdateInfo_VOffset);
+
+		UpdateScaleRotationOffset();
 
 		outPixels = JEnv->GetDirectBufferAddress(*buffer);
 		outCount = JEnv->GetDirectBufferCapacity(*buffer);
@@ -320,6 +366,83 @@ bool FJavaAndroidCameraPlayer::GetVideoLastFrameData(void* & outPixels, int64 &
 	return false;
 }
 
+/** This function has to be called before calling any GetVideoLastXXXXXXX() functions */
+void FJavaAndroidCameraPlayer::SelectBitmapRenderer()
+{
+	if (BRSelection == BitmapRendererSelection::Unknown)
+	{
+		if (UsingVulkanRenderer)
+		{
+			if (AndroidEnableCameraPlayerBRIR && FAndroidMisc::IsExternalMemoryAndroidHardwareBufferExtensionLoaded())
+			{
+				BRSelection = BitmapRendererSelection::ImageReader;
+			}
+			else
+			{
+				BRSelection = BitmapRendererSelection::Legacy;
+			}
+		}
+		else
+		{
+			if (AndroidEnableCameraPlayerBRIR)
+			{
+				BRSelection = BitmapRendererSelection::ImageReader;
+			}
+			else
+			{
+				BRSelection = BitmapRendererSelection::Legacy;
+			}
+		}
+
+		JNIEnv* JEnv = FAndroidApplication::GetJavaEnv();
+		JEnv->CallVoidMethod(Object, SetUsingBitmapRendererImageReaderMethod.Method, (BRSelection == BitmapRendererSelection::ImageReader));
+	}
+}
+
+bool FJavaAndroidCameraPlayer::IsUsingBitmapRendererImageReader()
+{
+	SelectBitmapRenderer();
+	return (BRSelection == BitmapRendererSelection::ImageReader);
+}
+
+jobject FJavaAndroidCameraPlayer::GetVideoLastMediaFrameData(int32& CurrentPosition, bool& bRegionChanged)
+{
+	bRegionChanged = false;
+
+	if (!IsUsingBitmapRendererImageReader())
+	{
+		return nullptr;
+	}
+
+	JNIEnv* JEnv = FAndroidApplication::GetJavaEnv();
+	jobject FrameData = JEnv->CallObjectMethod(Object, GetVideoLastMediaFrameDataMethod.Method);
+	auto FrameUpdateInfo = NewScopedJavaObject(JEnv, FrameData);
+	if (JEnv->ExceptionCheck())
+	{
+		JEnv->ExceptionDescribe();
+		JEnv->ExceptionClear();
+		return nullptr;
+	}
+
+	if (FrameUpdateInfo)
+	{
+		jobject BRFrameData = JEnv->GetObjectField(*FrameUpdateInfo, FrameUpdateInfo_BRFrameData);
+		if (BRFrameData)
+		{
+			CurrentPosition = (int32)JEnv->GetIntField(*FrameUpdateInfo, FrameUpdateInfo_CurrentPosition);
+			bool bFrameReady = (bool)JEnv->GetBooleanField(*FrameUpdateInfo, FrameUpdateInfo_FrameReady);
+			bRegionChanged = (bool)JEnv->GetBooleanField(*FrameUpdateInfo, FrameUpdateInfo_RegionChanged);
+
+			UpdateScaleRotationOffset();
+
+			jobject GlobalBRFrameData = FAndroidApplication::GetJavaEnv()->NewGlobalRef(BRFrameData);
+			return GlobalBRFrameData;
+		}
+	}
+
+	return nullptr;
+}
+
 void FJavaAndroidCameraPlayer::Start()
 {
 	CallMethod<void>(StartMethod);
@@ -364,18 +487,19 @@ bool FJavaAndroidCameraPlayer::UpdateVideoFrame(int32 ExternalTextureId, int32 *
 	*CurrentPosition = (int32)JEnv->GetIntField(*Result, FrameUpdateInfo_CurrentPosition);
 	bool bFrameReady = (bool)JEnv->GetBooleanField(*Result, FrameUpdateInfo_FrameReady);
 	*bRegionChanged = (bool)JEnv->GetBooleanField(*Result, FrameUpdateInfo_RegionChanged);
-	ScaleRotation.X = (float)(float)JEnv->GetFloatField(*Result, FrameUpdateInfo_ScaleRotation00);
-	ScaleRotation.Y = (float)(float)JEnv->GetFloatField(*Result, FrameUpdateInfo_ScaleRotation01);
-	ScaleRotation.Z = (float)(float)JEnv->GetFloatField(*Result, FrameUpdateInfo_ScaleRotation10);
-	ScaleRotation.W = (float)(float)JEnv->GetFloatField(*Result, FrameUpdateInfo_ScaleRotation11);
-	Offset.X = (float)JEnv->GetFloatField(*Result, FrameUpdateInfo_UOffset);
-	Offset.Y = (float)JEnv->GetFloatField(*Result, FrameUpdateInfo_VOffset);
+
+	UpdateScaleRotationOffset();
 
 	return bFrameReady;
 }
 
 bool FJavaAndroidCameraPlayer::GetVideoLastFrame(int32 destTexture)
 {
+	if (IsUsingBitmapRendererImageReader())
+	{
+		return false;
+	}
+
 	// This can return an exception in some cases
 	JNIEnv*	JEnv = FAndroidApplication::GetJavaEnv();
 	auto Result = NewScopedJavaObject(JEnv, JEnv->CallObjectMethod(Object, GetVideoLastFrameMethod.Method, destTexture));
@@ -392,12 +516,8 @@ bool FJavaAndroidCameraPlayer::GetVideoLastFrame(int32 destTexture)
 	}
 
 	bool bFrameReady = (bool)JEnv->GetBooleanField(*Result, FrameUpdateInfo_FrameReady);
-	ScaleRotation.X = (float)(float)JEnv->GetFloatField(*Result, FrameUpdateInfo_ScaleRotation00);
-	ScaleRotation.Y = (float)(float)JEnv->GetFloatField(*Result, FrameUpdateInfo_ScaleRotation01);
-	ScaleRotation.Z = (float)(float)JEnv->GetFloatField(*Result, FrameUpdateInfo_ScaleRotation10);
-	ScaleRotation.W = (float)(float)JEnv->GetFloatField(*Result, FrameUpdateInfo_ScaleRotation11);
-	Offset.X = (float)JEnv->GetFloatField(*Result, FrameUpdateInfo_UOffset);
-	Offset.Y = (float)JEnv->GetFloatField(*Result, FrameUpdateInfo_VOffset);
+
+	UpdateScaleRotationOffset();
 	
 	return bFrameReady;
 }
diff --git a/Engine/Plugins/Media/AndroidCamera/Source/AndroidCamera/Private/Player/AndroidCameraPlayer.cpp b/Engine/Plugins/Media/AndroidCamera/Source/AndroidCamera/Private/Player/AndroidCameraPlayer.cpp
--- a/Engine/Plugins/Media/AndroidCamera/Source/AndroidCamera/Private/Player/AndroidCameraPlayer.cpp
+++ b/Engine/Plugins/Media/AndroidCamera/Source/AndroidCamera/Private/Player/AndroidCameraPlayer.cpp
@@ -21,7 +21,6 @@
 #include "AndroidCameraTextureSample.h"
 
 #define ANDROIDCAMERAPLAYER_USE_EXTERNALTEXTURE 1
-#define ANDROIDCAMERAPLAYER_USE_PREPAREASYNC 0
 #define ANDROIDCAMERAPLAYER_USE_NATIVELOGGING 1
 
 #define LOCTEXT_NAMESPACE "FAndroidCameraPlayer"
@@ -46,6 +45,10 @@ FAndroidCameraPlayer::FAndroidCameraPlayer(IMediaEventSink& InEventSink)
 	, SelectedVideoTrack(INDEX_NONE)
 	, VideoSamplePool(new FAndroidCameraTextureSamplePool)
 	, bOpenWithoutEvents(false)
+	, IsInBackground(false)
+	, DeferredSetRateThread()
+	, IsSetRateThreadRunning(false)
+	, StopSetRateThread(false)
 {
 	check(JavaCameraPlayer.IsValid());
 	check(Samples.IsValid());
@@ -264,24 +267,8 @@ bool FAndroidCameraPlayer::Open(const FString& Url, const IMediaOptions* /*Optio
 	// prepare media
 	MediaUrl = JavaCameraPlayer->GetDataSource();
 
-#if ANDROIDCAMERAPLAYER_USE_PREPAREASYNC
-	if (!JavaCameraPlayer->PrepareAsync())
-	{
-		UE_LOG(LogAndroidCamera, Warning, TEXT("Failed to prepare media source %s"), *Url);
-		return false;
-	}
-
 	CurrentState = EMediaState::Preparing;
 	return true;
-#else
-	if (!JavaCameraPlayer->Prepare())
-	{
-		UE_LOG(LogAndroidCamera, Warning, TEXT("Failed to prepare media source %s"), *Url);
-		return false;
-	}
-
-	return InitializePlayer();
-#endif
 }
 
 
@@ -332,7 +319,7 @@ static void ConditionalInitializeVideoTexture_RenderThread(TWeakPtr<FJavaAndroid
 		PinnedJavaCameraPlayer->SetVideoTextureValid(false);
 
 #if ANDROIDCAMERAPLAYER_USE_NATIVELOGGING
-		FPlatformMisc::LowLevelOutputDebugStringf(TEXT("Fetch RT: Created VideoTexture: %d - %s"), *reinterpret_cast<int32*>(VideoTexture->GetNativeResource()), *PlayerGuid.ToString());
+		FPlatformMisc::LowLevelOutputDebugStringf(TEXT("Fetch RT: Created VideoTexture: %s"), *PlayerGuid.ToString());
 #endif
 	}
 }
@@ -512,6 +499,60 @@ struct FRHICommandUpdateBufferCameraSample final : public FRHICommand<FRHIComman
 	}
 };
 
+static void DoUpdateMediaFrameDataCameraSampleExecute(TWeakPtr<FJavaAndroidCameraPlayer, ESPMode::ThreadSafe> JavaCameraPlayerPtr, TWeakPtr<FMediaSamples, ESPMode::ThreadSafe> SamplesPtr,
+	TSharedRef<FAndroidCameraTextureSample, ESPMode::ThreadSafe> VideoSample, FTimespan SampleTime)
+{
+	auto PinnedJavaCameraPlayer = JavaCameraPlayerPtr.Pin();
+	auto PinnedSamples = SamplesPtr.Pin();
+
+	if (!PinnedJavaCameraPlayer.IsValid() || !PinnedSamples.IsValid())
+	{
+		return;
+	}
+
+	int32 CurrentFramePosition = 0;
+	bool bRegionChanged = false;
+
+	jobject MediaFrameData = PinnedJavaCameraPlayer->GetVideoLastMediaFrameData(CurrentFramePosition, bRegionChanged);
+	if (MediaFrameData == nullptr)
+	{
+		return;
+	}
+
+	if (PinnedJavaCameraPlayer->IsActive())
+	{
+		VideoSample->InitializeMediaFrameData(MediaFrameData, SampleTime);
+
+		FVector4 ScaleRotation = PinnedJavaCameraPlayer->GetScaleRotation();
+		FVector4 Offset = PinnedJavaCameraPlayer->GetOffset();
+		VideoSample->SetScaleRotationOffset(ScaleRotation, Offset);
+
+		PinnedSamples->AddVideo(VideoSample);
+	}
+}
+
+struct FRHICommandUpdateMediaFrameDataCameraSample final : public FRHICommand<FRHICommandUpdateMediaFrameDataCameraSample>
+{
+	TWeakPtr<FJavaAndroidCameraPlayer, ESPMode::ThreadSafe> JavaCameraPlayerPtr;
+	TWeakPtr<FMediaSamples, ESPMode::ThreadSafe> SamplesPtr;
+	TSharedRef<FAndroidCameraTextureSample, ESPMode::ThreadSafe> VideoSample;
+	FTimespan SampleTime;
+
+	FORCEINLINE_DEBUGGABLE FRHICommandUpdateMediaFrameDataCameraSample(TWeakPtr<FJavaAndroidCameraPlayer, ESPMode::ThreadSafe> InJavaCameraPlayerPtr, TWeakPtr<FMediaSamples, ESPMode::ThreadSafe> InSamplesPtr,
+		TSharedRef<FAndroidCameraTextureSample, ESPMode::ThreadSafe> InVideoSample, FTimespan InSampleTime)
+		: JavaCameraPlayerPtr(InJavaCameraPlayerPtr)
+		, SamplesPtr(InSamplesPtr)
+		, VideoSample(InVideoSample)
+		, SampleTime(InSampleTime)
+	{
+	}
+	void Execute(FRHICommandListBase& CmdList)
+	{
+		QUICK_SCOPE_CYCLE_COUNTER(STAT_FRHICommandUpdateBufferCameraSample_Execute);
+		DoUpdateMediaFrameDataCameraSampleExecute(JavaCameraPlayerPtr, SamplesPtr, VideoSample, SampleTime);
+	}
+};
+
 // ==============================================================
 
 void FAndroidCameraPlayer::TickFetch(FTimespan DeltaTime, FTimespan Timecode)
@@ -562,6 +603,16 @@ void FAndroidCameraPlayer::TickFetch(FTimespan DeltaTime, FTimespan Timecode)
 		}
 		Params = { JavaCameraPlayer, Samples, VideoSample, (int32)(VideoTrack.Dimensions.X * VideoTrack.Dimensions.Y * sizeof(int32)), CurrentTime };
 
+		if (JavaCameraPlayer->IsUsingBitmapRendererImageReader())
+		{
+			ENQUEUE_RENDER_COMMAND(AndroidCameraPlayerWriteVideoSample)(
+				[Params](FRHICommandListImmediate& RHICmdList)
+				{
+					DoUpdateMediaFrameDataCameraSampleExecute(Params.JavaCameraPlayerPtr, Params.SamplesPtr, Params.VideoSample, Params.SampleTime);
+				});
+			return;
+		}
+
 		ENQUEUE_RENDER_COMMAND(AndroidCameraPlayerWriteVideoSample)(
 			[Params](FRHICommandListImmediate& RHICmdList)
 			{
@@ -610,9 +661,38 @@ void FAndroidCameraPlayer::TickFetch(FTimespan DeltaTime, FTimespan Timecode)
 		{
 			TWeakPtr<FJavaAndroidCameraPlayer, ESPMode::ThreadSafe> JavaCameraPlayerPtr;
 			FGuid PlayerGuid;
+			TSharedRef<FAndroidCameraTextureSample, ESPMode::ThreadSafe> VideoSample;
+			TWeakPtr<FMediaSamples, ESPMode::ThreadSafe> Samples;
+			FTimespan SampleTime;
 		};
 
-		FWriteVideoSampleParams Params = { JavaCameraPlayer, PlayerGuid };
+		auto VideoSample = VideoSamplePool->AcquireShared();
+		const FJavaAndroidCameraPlayer::FVideoTrack VideoTrack = VideoTracks[SelectedVideoTrack];
+		if (!VideoSample->Initialize(VideoTrack.Dimensions, FTimespan::FromSeconds(1.0 / VideoTrack.FrameRate)))
+		{
+			return;
+		}
+
+		FWriteVideoSampleParams Params = { JavaCameraPlayer, PlayerGuid, VideoSample, Samples, CurrentTime };
+
+		if (JavaCameraPlayer->IsUsingBitmapRendererImageReader())
+		{
+
+			ENQUEUE_RENDER_COMMAND(AndroidCameraPlayerWriteVideoSample)(
+				[Params](FRHICommandListImmediate& RHICmdList)
+				{
+					if (IsRunningRHIInSeparateThread())
+					{
+						new (RHICmdList.AllocCommand<FRHICommandUpdateMediaFrameDataCameraSample>()) FRHICommandUpdateMediaFrameDataCameraSample(Params.JavaCameraPlayerPtr, Params.Samples, Params.VideoSample, Params.SampleTime);
+						RHICmdList.ImmediateFlush(EImmediateFlushType::FlushRHIThread);
+					}
+					else
+					{
+						DoUpdateMediaFrameDataCameraSampleExecute(Params.JavaCameraPlayerPtr, Params.Samples, Params.VideoSample, Params.SampleTime);
+					}
+				});
+			return;
+		}
 
 		ENQUEUE_RENDER_COMMAND(AndroidCameraPlayerWriteVideoSample)(
 			[Params](FRHICommandListImmediate& RHICmdList)
@@ -629,8 +709,9 @@ void FAndroidCameraPlayer::TickFetch(FTimespan DeltaTime, FTimespan Timecode)
 				}
 
 				auto PinnedJavaCameraPlayer = Params.JavaCameraPlayerPtr.Pin();
+				auto PinnedSamples = Params.Samples.Pin();
 
-				if (!PinnedJavaCameraPlayer.IsValid())
+				if (!PinnedJavaCameraPlayer.IsValid() || !PinnedSamples.IsValid())
 				{
 					return;
 				}
@@ -659,6 +740,25 @@ void FAndroidCameraPlayer::TickFetch(FTimespan DeltaTime, FTimespan Timecode)
 
 					PinnedJavaCameraPlayer->SetVideoTextureValid(true);
 				}
+
+				Params.VideoSample->InitializeExternalTexture(Params.SampleTime);
+				
+				// Update the associated data of the registered external texture
+				FTextureRHIRef SampleTexture;
+				FSamplerStateRHIRef SamplerState;
+				if (FExternalTextureRegistry::Get().GetExternalTexture(nullptr, Params.PlayerGuid, SampleTexture, SamplerState))
+				{
+					FVector4 ScaleRotation = PinnedJavaCameraPlayer->GetScaleRotation();
+					FVector4 Offset = PinnedJavaCameraPlayer->GetOffset();
+					FExternalTextureRegistry::Get().RegisterExternalTexture(Params.PlayerGuid, SampleTexture, SamplerState,
+						FLinearColor(ScaleRotation.X, ScaleRotation.Y, ScaleRotation.Z, ScaleRotation.W),
+						FLinearColor(Offset.X, Offset.Y, Offset.Z, Offset.W));
+					PinnedJavaCameraPlayer->SetVideoTextureValid(true);
+				}
+				
+				// Add a dummy video sample
+				// (actual data flows "on the side" via external texture mechanism)
+				PinnedSamples->AddVideo(Params.VideoSample);
 			});
 	}
 	else
@@ -724,6 +824,24 @@ void FAndroidCameraPlayer::TickFetch(FTimespan DeltaTime, FTimespan Timecode)
 	}
 	Params = { JavaCameraPlayer, Samples, VideoSample, (int32)(VideoTrack.Dimensions.X * VideoTrack.Dimensions.Y * sizeof(int32)), CurrentTime };
 
+	if (JavaCameraPlayer->IsUsingBitmapRendererImageReader())
+	{
+		ENQUEUE_RENDER_COMMAND(AndroidCameraPlayerWriteVideoSample)(
+			[Params](FRHICommandListImmediate& RHICmdList)
+			{
+				if (IsRunningRHIInSeparateThread())
+				{
+					new (RHICmdList.AllocCommand<FRHICommandUpdateMediaFrameDataCameraSample>()) FRHICommandUpdateMediaFrameDataCameraSample(Params.JavaCameraPlayerPtr, Params.SamplesPtr, Params.VideoSample, Params.SampleTime);
+					RHICmdList.ImmediateFlush(EImmediateFlushType::FlushRHIThread);
+				}
+				else
+				{
+					DoUpdateBufferCameraSampleExecute(Params.JavaCameraPlayerPtr, Params.SamplesPtr, Params.VideoSample, Params.SampleTime);
+				}
+			});
+		return;
+	}
+
 	ENQUEUE_RENDER_COMMAND(AndroidCameraPlayerWriteVideoSample)(
 		[Params](FRHICommandListImmediate& RHICmdList)
 		{
@@ -746,7 +864,7 @@ void FAndroidCameraPlayer::TickInput(FTimespan DeltaTime, FTimespan Timecode)
 	if (CurrentState != EMediaState::Playing)
 	{
 		// remove delegates if registered
-		if (ResumeHandle.IsValid())
+		if (ResumeHandle.IsValid() && !IsInBackground)
 		{
 			FCoreDelegates::ApplicationHasEnteredForegroundDelegate.Remove(ResumeHandle);
 			ResumeHandle.Reset();
@@ -758,22 +876,6 @@ void FAndroidCameraPlayer::TickInput(FTimespan DeltaTime, FTimespan Timecode)
 			PauseHandle.Reset();
 		}
 
-#if ANDROIDCAMERAPLAYER_USE_PREPAREASYNC
-		// if preparing, see if finished
-		if (CurrentState == EMediaState::Preparing)
-		{
-			if (!JavaCameraPlayer.IsValid())
-			{
-				return;
-			}
-
-			if (JavaCameraPlayer->IsPrepared())
-			{
-				InitializePlayer();
-			}
-		}
-#endif
-
 		return;
 	}
 
@@ -838,17 +940,17 @@ bool FAndroidCameraPlayer::InitializePlayer()
 		FPlatformMisc::LowLevelOutputDebugStringf(TEXT("FAndroidCamera::InitializePlayer %s"), *PlayerGuid.ToString());
 	#endif
 
-	if (ResumeHandle.IsValid())
-	{
-		FCoreDelegates::ApplicationHasEnteredForegroundDelegate.Remove(ResumeHandle);
-		ResumeHandle.Reset();
-	}
+		if (ResumeHandle.IsValid())
+		{
+			FCoreDelegates::ApplicationHasEnteredForegroundDelegate.Remove(ResumeHandle);
+			ResumeHandle.Reset();
+		}
 
-	if (PauseHandle.IsValid())
-	{
-		FCoreDelegates::ApplicationWillEnterBackgroundDelegate.Remove(PauseHandle);
-		PauseHandle.Reset();
-	}
+		if (PauseHandle.IsValid())
+		{
+			FCoreDelegates::ApplicationWillEnterBackgroundDelegate.Remove(PauseHandle);
+			PauseHandle.Reset();
+		}
 
 	JavaCameraPlayer->GetAudioTracks(AudioTracks);
 	JavaCameraPlayer->GetCaptionTracks(CaptionTracks);
@@ -1507,13 +1609,57 @@ bool FAndroidCameraPlayer::SetLooping(bool Looping)
 bool FAndroidCameraPlayer::SetRate(float Rate)
 {
 	if ((CurrentState == EMediaState::Closed) ||
-		(CurrentState == EMediaState::Error) ||
-		(CurrentState == EMediaState::Preparing))
+		(CurrentState == EMediaState::Error))
 	{
 		UE_LOG(LogAndroidCamera, Warning, TEXT("Cannot set rate while closed, preparing, or in error state"));
 		return false;
 	}
 
+	if (CurrentState == EMediaState::Preparing && Rate == 1.0f)
+	{
+		if (IsSetRateThreadRunning)
+		{
+			UE_LOG(LogAndroidCamera, Warning, TEXT("It's unnecessary to Play the camera player repeatedly."));
+			return true;
+		}
+
+		UE_LOG(LogAndroidCamera, Log, TEXT("Camera is under preparation."));
+
+		StopSetRateThread = false;
+
+		int ret = pthread_create(&DeferredSetRateThread, NULL, [] (void* arg) -> void* {
+			FAndroidCameraPlayer* Player = (FAndroidCameraPlayer*)arg;
+			Player->SetSetRateThreadStatus(true);
+			UE_LOG(LogAndroidCamera, Log, TEXT("Waiting for camera ready..."));
+			float ExitTime = FPlatformTime::Seconds() + 3.0f;
+			while (true)
+			{
+				if (FPlatformTime::Seconds() > ExitTime)
+				{
+					UE_LOG(LogAndroidCamera, Log, TEXT("Timed out after waiting for camera preparation for 3 seconds."));
+					break;
+				}
+				if (Player->IsStoppingSetRateThread())
+				{
+					UE_LOG(LogAndroidCamera, Log, TEXT("Stopped waiting for camera preparation."));
+					break;
+				}
+				if (Player->JavaCameraPlayer->IsPrepared())
+				{
+					UE_LOG(LogAndroidCamera, Log, TEXT("Camera is ready."));
+					Player->InitializePlayer();
+					Player->SetRate(1.0f);
+					break;
+				}
+				usleep(100);
+			}
+			Player->SetSetRateThreadStatus(false);
+			return nullptr;
+		}, this);
+		
+		return (ret == 0);
+	}
+
 	if (Rate == GetRate())
 	{
 		return true; // rate already set
@@ -1523,6 +1669,14 @@ bool FAndroidCameraPlayer::SetRate(float Rate)
 
 	if (Rate == 0.0f)
 	{
+		if (IsSetRateThreadRunning)
+		{
+			StopSetRateThread = true;
+			pthread_join(DeferredSetRateThread, nullptr);
+			StopSetRateThread = false;
+			IsSetRateThreadRunning = false;
+		}
+
 		JavaCameraPlayer->Pause();
 		CurrentState = EMediaState::Paused;
 		EventSink.ReceiveMediaEvent(EMediaEvent::PlaybackSuspended);
@@ -1548,25 +1702,35 @@ bool FAndroidCameraPlayer::SetRate(float Rate)
 
 void FAndroidCameraPlayer::HandleApplicationWillEnterBackground()
 {
-	/*
+	// prevent the HandleApplicationHasEnteredForeground from removed in TickInput()
+	IsInBackground = true;
+
 	// check state in case changed before ticked
-	if ((CurrentState == EMediaState::Playing) && JavaCameraPlayer.IsValid())
+	if (JavaCameraPlayer.IsValid())
 	{
-		JavaCameraPlayer->Pause();
+		JavaCameraPlayer->Release();
+		CurrentState = EMediaState::Closed;
 	}
-	*/
 }
 
 
 void FAndroidCameraPlayer::HandleApplicationHasEnteredForeground()
 {
-	/*
+	IsInBackground = false;
+
 	// check state in case changed before ticked
-	if ((CurrentState == EMediaState::Playing) && JavaCameraPlayer.IsValid())
+	if (JavaCameraPlayer.IsValid() && CurrentState == EMediaState::Closed)
 	{
-		JavaCameraPlayer->Start();
+		if (JavaCameraPlayer->RestoreCamera())
+		{
+			CurrentState = EMediaState::Preparing;
+			SetRate(1.0f);
+		}
+		else
+		{
+			UE_LOG(LogAndroidCamera, Warning, TEXT("HandleApplicationHasEnteredForeground - RestoreCamera failed."));
+		}
 	}
-	*/
 }
 
 
diff --git a/Engine/Plugins/Media/AndroidCamera/Source/AndroidCamera/Private/Player/AndroidCameraPlayer.h b/Engine/Plugins/Media/AndroidCamera/Source/AndroidCamera/Private/Player/AndroidCameraPlayer.h
--- a/Engine/Plugins/Media/AndroidCamera/Source/AndroidCamera/Private/Player/AndroidCameraPlayer.h
+++ b/Engine/Plugins/Media/AndroidCamera/Source/AndroidCamera/Private/Player/AndroidCameraPlayer.h
@@ -3,6 +3,7 @@
 #pragma once
 
 #include "CoreTypes.h"
+#include <pthread.h>
 #include "Containers/Array.h"
 #include "Containers/UnrealString.h"
 #include "Misc/Guid.h"
@@ -63,6 +64,9 @@ public:
 	virtual void TickFetch(FTimespan DeltaTime, FTimespan Timecode) override;
 	virtual void TickInput(FTimespan DeltaTime, FTimespan Timecode) override;
 
+	bool IsStoppingSetRateThread() { return StopSetRateThread; }
+	void SetSetRateThreadStatus(bool IsRunning) { IsSetRateThreadRunning = IsRunning; }
+
 protected:
 
 	/**
@@ -173,4 +177,14 @@ private:
 
 	/** Whether or not the current open request should send events on completion. */
 	bool bOpenWithoutEvents;
+
+	/** Identify the game is in background */
+	bool IsInBackground;
+
+	/** The thread of delaying the SetRate until the camera is prepared */
+	pthread_t DeferredSetRateThread;
+	TAtomic<bool> IsSetRateThreadRunning;
+	
+	/** The flag for exiting the SetRate thread */
+	TAtomic<bool> StopSetRateThread;
 };
diff --git a/Engine/Plugins/Media/AndroidCamera/Source/AndroidCamera/Private/Player/AndroidCameraTextureSample.h b/Engine/Plugins/Media/AndroidCamera/Source/AndroidCamera/Private/Player/AndroidCameraTextureSample.h
--- a/Engine/Plugins/Media/AndroidCamera/Source/AndroidCamera/Private/Player/AndroidCameraTextureSample.h
+++ b/Engine/Plugins/Media/AndroidCamera/Source/AndroidCamera/Private/Player/AndroidCameraTextureSample.h
@@ -12,17 +12,22 @@
 #include "RHI.h"
 #include "RHIUtilities.h"
 #include "Templates/SharedPointer.h"
+#include "RHIStaticStates.h"
 
 #if !WITH_ENGINE
 	#include "MediaObjectPool.h"
 #endif
 
+#include "IMediaTextureSampleConverter.h"
+#include "Android/AndroidJava.h"
+#include "Android/AndroidJavaMediaFrameData.h"
 
 /**
  * Texture sample generated by AndroidCamera player.
  */
 class FAndroidCameraTextureSample
 	: public IMediaTextureSample
+	, public IMediaTextureSampleConverter
 	, public IMediaPoolable
 {
 public:
@@ -41,10 +46,7 @@ public:
 	/** Virtual destructor. */
 	virtual ~FAndroidCameraTextureSample()
 	{
-		if (BufferSize > 0)
-		{
-			FMemory::Free(Buffer);
-		}
+		ShutdownPoolable();
 	}
 
 public:
@@ -79,6 +81,15 @@ public:
 		return true;
 	}
 
+	/**
+	 * Initialize the sample as an external image reference.
+	 */
+	void InitializeExternalTexture(FTimespan InTime)
+	{
+		Time = InTime;
+		ExternalTexture = true;
+	}
+
 	/**
 	 * Initialize the sample with a memory buffer.
 	 *
@@ -154,6 +165,21 @@ public:
 		return Texture;
 	}
 
+	/**
+	 * Initialize the sample with a hardware buffer extract from MediaFrameData.
+	 *
+	 * @param InFrameData holds the jobject of the sample data.
+	 * @param InTime The sample time (in the player's local clock).
+	 * @note This method must be called on the render thread.
+	 * @see InitializeBuffer, InitializeTexture
+	 */
+	void InitializeMediaFrameData(jobject InFrameData, FTimespan InTime)
+	{
+		MediaFrameData.CleanUp();
+		MediaFrameData.Set(InFrameData);
+		Time = InTime;
+	}
+
 	/**
 	 * Set the sample Scale, Rotation, Offset.
 	 *
@@ -214,6 +240,11 @@ public:
 		return FMediaTimeStamp(Time);
 	}
 
+	virtual bool IsExternalImage() const override
+	{
+		return ExternalTexture;
+	}
+
 	virtual bool IsCacheable() const override
 	{
 #if WITH_ENGINE
@@ -238,14 +269,167 @@ public:
 		return Offset;
 	}
 
+	virtual EMediaOrientation GetOrientation() const override
+	{
+		return EMediaOrientation::CW90;
+	}
+
+	//~ IMediaTextureSampleConverter interface
+	virtual IMediaTextureSampleConverter* GetMediaTextureSampleConverter() override
+	{
+		if (MediaFrameData || Buffer)
+		{
+			return this;
+		}
+
+		return nullptr;
+	}
+
+	//~ IMediaPoolable interface
+	virtual bool IsReadyForReuse() override
+	{
+		if (MediaFrameData)
+		{
+			return MediaFrameData.IsReadyToClean();
+		}
+		return true;
+	}
+
+	virtual void ShutdownPoolable() override
+	{
+		// Reuse the texture?
+		//Texture = nullptr;
+
+		if (BufferSize > 0)
+		{
+			FMemory::Free(Buffer);
+			BufferSize = 0;
+			Buffer = nullptr;
+		}
+
+		MediaFrameData.CleanUp();
+	}
+
 private:
 
+	bool RotateBinaryBufferTexture(FRHICommandListImmediate& RHICmdList, FTextureRHIRef& InDstTexture, const FConversionHints& Hints)
+	{
+		const FRHITextureCreateDesc Desc =
+			FRHITextureCreateDesc::Create2D(TEXT("CameraSampleConverterIntermediateTexture"))
+			.SetExtent(Dim)
+			.SetFormat(PF_B8G8R8A8)
+			.SetNumMips(Hints.NumMips)
+			.DetermineInititialState();
+		FTextureRHIRef IntermediateTexture = RHICmdList.CreateTexture(Desc);
+
+		if (!IntermediateTexture.IsValid())
+		{
+			return false;
+		}
+
+		// Copy sample data (from CPU mem) to output render target
+		const FPixelFormatInfo& FormatInfo = GPixelFormats[PF_B8G8R8A8];
+		const uint8* Data = (const uint8*)Buffer;
+		for (uint8 MipLevel = 0; MipLevel < Hints.NumMips; ++MipLevel)
+		{
+			uint32 Width = Dim.X >> MipLevel;
+			uint32 Height = Dim.Y >> MipLevel;
+			uint32 Stride = GetStride() >> MipLevel;
+
+			FUpdateTextureRegion2D Region(0, 0, 0, 0, Width, Height);
+			RHICmdList.UpdateTexture2D(IntermediateTexture, MipLevel, Region, Stride, Data);
+
+			uint32 HeightUnit = (Height + FormatInfo.BlockSizeY - 1) / FormatInfo.BlockSizeY;
+			Data += Stride * HeightUnit;
+		}
+
+		FGraphicsPipelineStateInitializer GraphicsPSOInit;
+		FRHITexture* RenderTarget = InDstTexture.GetReference();
+
+		RHICmdList.Transition(FRHITransitionInfo(InDstTexture, ERHIAccess::Unknown, ERHIAccess::RTV));
+
+		FRHIRenderPassInfo RPInfo(RenderTarget, ERenderTargetActions::DontLoad_Store);
+		RHICmdList.BeginRenderPass(RPInfo, TEXT("ConvertMedia_BinaryBuffer"));
+		{
+			RHICmdList.ApplyCachedRenderTargets(GraphicsPSOInit);
+
+			RHICmdList.SetViewport(0, 0, 0.0f, (float)Dim.X, (float)Dim.Y, 1.0f);
+
+			GraphicsPSOInit.DepthStencilState = TStaticDepthStencilState<false, CF_Always>::GetRHI();
+			GraphicsPSOInit.RasterizerState = TStaticRasterizerState<>::GetRHI();
+			GraphicsPSOInit.BlendState = TStaticBlendStateWriteMask<CW_RGBA, CW_NONE, CW_NONE, CW_NONE, CW_NONE, CW_NONE, CW_NONE, CW_NONE>::GetRHI();
+			GraphicsPSOInit.PrimitiveType = PT_TriangleStrip;
+
+			// configure media shaders
+			auto ShaderMap = GetGlobalShaderMap(GMaxRHIFeatureLevel);
+			TShaderMapRef<FMediaShadersVS> VertexShader(ShaderMap);
+
+			GraphicsPSOInit.BoundShaderState.VertexDeclarationRHI = GMediaVertexDeclaration.VertexDeclarationRHI;
+			GraphicsPSOInit.BoundShaderState.VertexShaderRHI = VertexShader.GetVertexShader();
+
+			FSamplerStateInitializerRHI InSamplerStateInitializer(SF_Bilinear, AM_Clamp, AM_Clamp, AM_Clamp);
+			FSamplerStateRHIRef InSamplerState = RHICreateSamplerState(InSamplerStateInitializer);
+
+			TShaderMapRef<FReadTextureExternalPS> CopyShader(ShaderMap);
+			GraphicsPSOInit.BoundShaderState.PixelShaderRHI = CopyShader.GetPixelShader();
+			SetGraphicsPipelineState(RHICmdList, GraphicsPSOInit, 0);
+			SetShaderParametersLegacyPS(RHICmdList, CopyShader, IntermediateTexture, InSamplerState, ScaleRotation, Offset);
+
+			// draw full size quad into render target
+			FBufferRHIRef VertexBuffer = CreateTempMediaVertexBuffer(RHICmdList);
+			RHICmdList.SetStreamSource(0, VertexBuffer, 0);
+			// set viewport to RT size
+			RHICmdList.SetViewport(0, 0, 0.0f, (float)Dim.X, (float)Dim.Y, 1.0f);
+
+			RHICmdList.DrawPrimitive(0, 2, 1);
+		}
+		RHICmdList.EndRenderPass();
+
+		return true;
+	}
+
+	//~ IMediaTextureSampleConverter overrides
+	virtual bool Convert(FRHICommandListImmediate& RHICmdList, FTextureRHIRef& InDstTexture, const FConversionHints& Hints) override
+	{
+		if (MediaFrameData)
+		{
+			if (FAndroidMisc::ShouldUseVulkan())
+			{
+				return MediaFrameData.ExtractToTextureVulkan(RHICmdList, InDstTexture, this);
+			}
+			else
+			{
+				return MediaFrameData.ExtractToTextureOES(RHICmdList, InDstTexture, this);
+			}
+		}
+
+		if (Buffer != nullptr && BufferSize > 0)
+		{
+			if (FAndroidMisc::ShouldUseVulkan())
+			{
+				return RotateBinaryBufferTexture(RHICmdList, InDstTexture, Hints);
+			}
+			else
+			{
+				UE_LOG(LogAndroidCamera, Warning, TEXT("The binary buffer path for AndroidCameraTextureSample doesn't work without Vulkan."));
+			}
+		}
+
+		return false;
+	}
+
+	/** The sample's frame data from the java surface provider. */
+	FAndroidJavaMediaFrameData MediaFrameData;
+
 	/** The sample's data buffer. */
 	void* Buffer;
 
 	/** Current allocation size of Buffer. */
 	SIZE_T BufferSize;
 
+	/** Flag indicating if this is an external image reference. */
+	bool ExternalTexture;
+
 	/** Width and height of the texture sample. */
 	FIntPoint Dim;
 
diff --git a/Engine/Plugins/Media/AndroidCamera/Source/AndroidCamera/Public/AndroidJavaCameraPlayer.h b/Engine/Plugins/Media/AndroidCamera/Source/AndroidCamera/Public/AndroidJavaCameraPlayer.h
--- a/Engine/Plugins/Media/AndroidCamera/Source/AndroidCamera/Public/AndroidJavaCameraPlayer.h
+++ b/Engine/Plugins/Media/AndroidCamera/Source/AndroidCamera/Public/AndroidJavaCameraPlayer.h
@@ -77,12 +77,15 @@ public:
 	bool PrepareAsync();
 	void SeekTo(int32 Milliseconds);
 	void SetLooping(bool Looping);
+	bool RestoreCamera();
 	void Release();
 	int32 GetVideoHeight();
 	int32 GetVideoWidth();
 	float GetFrameRate();
 	void SetVideoEnabled(bool enabled = true);
 	void SetAudioEnabled(bool enabled = true);
+	bool IsUsingBitmapRendererImageReader();
+	jobject GetVideoLastMediaFrameData(int32& CurrentPosition, bool& bRegionChanged);
 	bool GetVideoLastFrameData(void* & outPixels, int64 & outCount, int32 *CurrentPosition, bool *bRegionChanged);
 	void Start();
 	void Pause();
@@ -118,6 +121,7 @@ private:
 	FJavaClassMethod PrepareAsyncMethod;
 	FJavaClassMethod SeekToMethod;
 	FJavaClassMethod SetLoopingMethod;
+	FJavaClassMethod OpenCameraMethod;
 	FJavaClassMethod ReleaseMethod;
 	FJavaClassMethod GetVideoHeightMethod;
 	FJavaClassMethod GetVideoWidthMethod;
@@ -137,6 +141,10 @@ private:
 	FJavaClassMethod UpdateVideoFrameMethod;
 	FJavaClassMethod TakePictureMethod;
 
+	// interfaces of getting video frame data through ImageReader
+	FJavaClassMethod SetUsingBitmapRendererImageReaderMethod;
+	FJavaClassMethod GetVideoLastMediaFrameDataMethod;
+
 	// FrameUpdateInfo member field ids
 	jclass FrameUpdateInfoClass;
 	jfieldID FrameUpdateInfo_Buffer;
@@ -149,6 +157,7 @@ private:
 	jfieldID FrameUpdateInfo_ScaleRotation11;
 	jfieldID FrameUpdateInfo_UOffset;
 	jfieldID FrameUpdateInfo_VOffset;
+	jfieldID FrameUpdateInfo_BRFrameData;
 
 	// AudioDeviceInfo member field ids
 	jclass AudioTrackInfoClass;
@@ -185,6 +194,19 @@ private:
 	FVector4 ScaleRotation;
 	FVector4 Offset;
 
+	enum BitmapRendererSelection
+	{
+		Unknown,
+		Legacy,
+		ImageReader
+	};
+
+	BitmapRendererSelection BRSelection; // BR - BitmapRenderer
+	bool UsingVulkanRenderer;
+
+	void SelectBitmapRenderer();
+	void UpdateScaleRotationOffset();
+
 public:
 	FTextureRHIRef GetVideoTexture()
 	{
diff --git a/Engine/Plugins/Media/AndroidMedia/Source/AndroidMedia/Private/Player/AndroidMediaPlayer.cpp b/Engine/Plugins/Media/AndroidMedia/Source/AndroidMedia/Private/Player/AndroidMediaPlayer.cpp
--- a/Engine/Plugins/Media/AndroidMedia/Source/AndroidMedia/Private/Player/AndroidMediaPlayer.cpp
+++ b/Engine/Plugins/Media/AndroidMedia/Source/AndroidMedia/Private/Player/AndroidMediaPlayer.cpp
@@ -849,7 +849,16 @@ void FAndroidMediaPlayer::TickFetch(FTimespan DeltaTime, FTimespan /*Timecode*/)
 			ENQUEUE_RENDER_COMMAND(AndroidMediaPlayerWriteVideoSample)(
 				[Params](FRHICommandListImmediate& RHICmdList)
 				{
-					DoUpdateMediaFrameDataMediaSampleExecute(Params.JavaMediaPlayerPtr, Params.Samples, Params.VideoSample);
+					if (IsRunningRHIInSeparateThread())
+					{
+						new (RHICmdList.AllocCommand<FRHICommandUpdateMediaFrameDataMediaSample>()) FRHICommandUpdateMediaFrameDataMediaSample(Params.JavaMediaPlayerPtr, Params.Samples, Params.VideoSample);
+						// wait for DoUpdateExternalMediaSampleExecute to complete.
+						RHICmdList.ImmediateFlush(EImmediateFlushType::FlushRHIThread);
+					}
+					else
+					{
+						DoUpdateMediaFrameDataMediaSampleExecute(Params.JavaMediaPlayerPtr, Params.Samples, Params.VideoSample);
+					}
 				});
 		}
 	}
@@ -937,7 +946,7 @@ void FAndroidMediaPlayer::TickFetch(FTimespan DeltaTime, FTimespan /*Timecode*/)
 			{
 				if (IsRunningRHIInSeparateThread())
 				{
-					new (RHICmdList.AllocCommand<FRHICommandUpdateHardwareBufferMediaSample>()) FRHICommandUpdateHardwareBufferMediaSample(Params.JavaMediaPlayerPtr, Params.SamplesPtr, Params.VideoSample);
+					new (RHICmdList.AllocCommand<FRHICommandUpdateMediaFrameDataMediaSample>()) FRHICommandUpdateMediaFrameDataMediaSample(Params.JavaMediaPlayerPtr, Params.SamplesPtr, Params.VideoSample);
 				}
 				else
 				{
diff --git a/Engine/Shaders/Private/MediaShaders.usf b/Engine/Shaders/Private/MediaShaders.usf
--- a/Engine/Shaders/Private/MediaShaders.usf
+++ b/Engine/Shaders/Private/MediaShaders.usf
@@ -618,7 +618,8 @@ void VYUConvertPS(
 	float2 InUV : TEXCOORD0,
 	out float4 OutColor : SV_Target0)
 {
-	float2 ScaledUV = InUV * VYUConvertUB.UVScale;
+	float2 RotatedUV = float2(dot(InUV, VYUConvertUB.UVRotation.xy), dot(InUV, VYUConvertUB.UVRotation.zw)) + VYUConvertUB.UVRotOffsetAndScale.xy;
+	float2 ScaledUV = RotatedUV * VYUConvertUB.UVRotOffsetAndScale.zw;
 	float3 YUV = VYUConvertUB.Texture.Sample(VYUConvertUB.Sampler, ScaledUV).gbr;
 
 	// Output with solid alpha
diff --git a/Engine/Source/Runtime/ApplicationCore/Private/Android/AndroidJavaMediaFrameData.cpp b/Engine/Source/Runtime/ApplicationCore/Private/Android/AndroidJavaMediaFrameData.cpp
--- a/Engine/Source/Runtime/ApplicationCore/Private/Android/AndroidJavaMediaFrameData.cpp
+++ b/Engine/Source/Runtime/ApplicationCore/Private/Android/AndroidJavaMediaFrameData.cpp
@@ -73,7 +73,7 @@ bool FAndroidJavaMediaFrameData::Set(jobject InFrameData)
 	return (FrameDataGlobalRef != nullptr);
 }
 
-jobject FAndroidJavaMediaFrameData::Extract(FVector2f& OutScale, FVector2f& OutOffset)
+jobject FAndroidJavaMediaFrameData::Extract()
 {
 	if (FrameDataGlobalRef == nullptr)
 	{
@@ -81,11 +81,6 @@ jobject FAndroidJavaMediaFrameData::Extract(FVector2f& OutScale, FVector2f& OutO
 	}
 
 	JNIEnv* JEnv = FAndroidApplication::GetJavaEnv();
-	OutScale.X = JEnv->GetFloatField(FrameDataGlobalRef, FrameData_UScale);
-	OutScale.Y = JEnv->GetFloatField(FrameDataGlobalRef, FrameData_VScale);
-	OutOffset.X = JEnv->GetFloatField(FrameDataGlobalRef, FrameData_UOffset);
-	OutOffset.Y = JEnv->GetFloatField(FrameDataGlobalRef, FrameData_VOffset);
-
 	jobject HardwareBufferObj = JEnv->GetObjectField(FrameDataGlobalRef, FrameData_HardwareBufferHandle);
 	check(HardwareBufferObj != nullptr);
 
@@ -115,10 +110,8 @@ bool FAndroidJavaMediaFrameData::ExtractToTextureVulkan(FRHICommandListImmediate
 	check(IsInRenderingThread());
 
 	bool Converted = false;
-	FVector2f FrameDataScale;
-	FVector2f FrameDataOffset;
 
-	jobject HardwareBufferObj = Extract(FrameDataScale, FrameDataOffset);
+	jobject HardwareBufferObj = Extract();
 	AHardwareBuffer* HardwareBuffer = AHardwareBuffer_fromHardwareBuffer(FAndroidApplication::GetJavaEnv(), HardwareBufferObj);
 	if (HardwareBuffer != nullptr)
 	{
@@ -182,7 +175,8 @@ bool FAndroidJavaMediaFrameData::ExtractToTextureVulkan(FRHICommandListImmediate
 				TShaderMapRef<FVYUConvertPS> ConvertShader(ShaderMap);
 				GraphicsPSOInit.BoundShaderState.PixelShaderRHI = ConvertShader.GetPixelShader();
 				SetGraphicsPipelineState(RHICmdList, GraphicsPSOInit, 0);
-				SetShaderParametersLegacyPS(RHICmdList, ConvertShader, SrcTexture, Dim, YUVMtx, TextureSample->GetEncodingType(), ColorSpaceMtx);
+				SetShaderParametersLegacyPS(RHICmdList, ConvertShader, SrcTexture, Dim, YUVMtx, TextureSample->GetEncodingType(), ColorSpaceMtx, 
+					MediaShaders::EToneMapMethod::None, TextureSample->GetScaleRotation(), TextureSample->GetOffset());
 
 				// draw full size quad into render target
 				FBufferRHIRef VertexBuffer = CreateTempMediaVertexBuffer(RHICmdList);
@@ -218,10 +212,8 @@ bool FAndroidJavaMediaFrameData::ExtractToTextureOES(FRHICommandListImmediate& R
 	check(IsInRenderingThread());
 
 	bool Converted = false;
-	FVector2f FrameDataScale;
-	FVector2f FrameDataOffset;
 
-	jobject HardwareBufferObj = Extract(FrameDataScale, FrameDataOffset);
+	jobject HardwareBufferObj = Extract();
 	AHardwareBuffer* HardwareBuffer = AHardwareBuffer_fromHardwareBuffer(FAndroidApplication::GetJavaEnv(), HardwareBufferObj);
 	if (HardwareBuffer != nullptr)
 	{
@@ -231,10 +223,6 @@ bool FAndroidJavaMediaFrameData::ExtractToTextureOES(FRHICommandListImmediate& R
 		FTextureRHIRef SrcTexture = RHI->RHICreateTexture2DFromAndroidHardwareBuffer(RHICmdList, HardwareBuffer);
 		if (SrcTexture.IsValid())
 		{
-			FLinearColor Offset = { FrameDataOffset.X, FrameDataOffset.Y, 0.f, 0.f };
-			//FLinearColor ScaleRotation = { FrameDataScale.X, 0.f, 0.f, FrameDataScale.Y };
-			FLinearColor ScaleRotation = { 1.f, 0.f, 0.f, 1.f };
-
 			FGraphicsPipelineStateInitializer GraphicsPSOInit;
 			FRHITexture* RenderTarget = InDstTexture.GetReference();
 
@@ -266,7 +254,7 @@ bool FAndroidJavaMediaFrameData::ExtractToTextureOES(FRHICommandListImmediate& R
 				TShaderMapRef<FReadTextureExternalPS> CopyShader(ShaderMap);
 				GraphicsPSOInit.BoundShaderState.PixelShaderRHI = CopyShader.GetPixelShader();
 				SetGraphicsPipelineState(RHICmdList, GraphicsPSOInit, 0);
-				SetShaderParametersLegacyPS(RHICmdList, CopyShader, SrcTexture, InSamplerState, ScaleRotation, Offset);
+				SetShaderParametersLegacyPS(RHICmdList, CopyShader, SrcTexture, InSamplerState, TextureSample->GetScaleRotation(), TextureSample->GetOffset());
 
 				// draw full size quad into render target
 				FBufferRHIRef VertexBuffer = CreateTempMediaVertexBuffer(RHICmdList);
diff --git a/Engine/Source/Runtime/ApplicationCore/Public/Android/AndroidJavaMediaFrameData.h b/Engine/Source/Runtime/ApplicationCore/Public/Android/AndroidJavaMediaFrameData.h
--- a/Engine/Source/Runtime/ApplicationCore/Public/Android/AndroidJavaMediaFrameData.h
+++ b/Engine/Source/Runtime/ApplicationCore/Public/Android/AndroidJavaMediaFrameData.h
@@ -37,7 +37,7 @@ public:
 	bool Set(jobject InFrameData);
 	void CleanUp();
 
-	jobject Extract(FVector2f& OutScale, FVector2f& OutOffset);
+	jobject Extract();
 	bool ExtractToTextureVulkan(FRHICommandListImmediate& RHICmdList, FTextureRHIRef& InDstTexture, IMediaTextureSample* TextureSample);
 	bool ExtractToTextureOES(FRHICommandListImmediate& RHICmdList, FTextureRHIRef& InDstTexture, IMediaTextureSample* TextureSample);
 
diff --git a/Engine/Source/Runtime/Core/Private/Android/AndroidPlatformMisc.cpp b/Engine/Source/Runtime/Core/Private/Android/AndroidPlatformMisc.cpp
--- a/Engine/Source/Runtime/Core/Private/Android/AndroidPlatformMisc.cpp
+++ b/Engine/Source/Runtime/Core/Private/Android/AndroidPlatformMisc.cpp
@@ -2375,15 +2375,15 @@ bool FAndroidMisc::IsExternalMemoryAndroidHardwareBufferExtensionLoaded()
 		// Consider replace the following check by using PhysicalDeviceFeatures.Core_1_1.samplerYcbcrConversion
 		if (!LoadedDeviceExtensions.Find(FAnsiString(VK_KHR_SAMPLER_YCBCR_CONVERSION_EXTENSION_NAME)))
 		{
-			UE_LOG(LogAndroid, Log, TEXT("Selecting CPU path because GPU extension '" VK_KHR_SAMPLER_YCBCR_CONVERSION_EXTENSION_NAME "' is not available!"));
+			UE_LOG(LogAndroid, Log, TEXT("VK_KHR_SAMPLER_YCBCR_CONVERSION_EXTENSION is not available!"));
 		}
 		else if (!LoadedDeviceExtensions.Find(FAnsiString(VK_ANDROID_EXTERNAL_MEMORY_ANDROID_HARDWARE_BUFFER_EXTENSION_NAME)))
 		{
-			UE_LOG(LogAndroid, Log, TEXT("Selecting CPU path because GPU extension '" VK_ANDROID_EXTERNAL_MEMORY_ANDROID_HARDWARE_BUFFER_EXTENSION_NAME "' is not available!"));
+			UE_LOG(LogAndroid, Log, TEXT("VK_ANDROID_EXTERNAL_MEMORY_ANDROID_HARDWARE_BUFFER_EXTENSION is not available!"));
 		}
 		else
 		{
-			UE_LOG(LogAndroid, Log, TEXT("Selecting GPU path because it is enabled via Electra.AndroidUseGpuOutputPath = 1"));
+			UE_LOG(LogAndroid, Log, TEXT("VK_ANDROID_EXTERNAL_MEMORY_ANDROID_HARDWARE_BUFFER_EXTENSION is loaded."));
 			return true;
 		}
 	}
diff --git a/Engine/Source/Runtime/MediaAssets/Private/Misc/MediaTextureResource.cpp b/Engine/Source/Runtime/MediaAssets/Private/Misc/MediaTextureResource.cpp
--- a/Engine/Source/Runtime/MediaAssets/Private/Misc/MediaTextureResource.cpp
+++ b/Engine/Source/Runtime/MediaAssets/Private/Misc/MediaTextureResource.cpp
@@ -1480,22 +1480,28 @@ void FMediaTextureResource::CopySample(FRHICommandListImmediate& RHICmdList, con
 		}
 		else
 		{
-			// Copy sample data (from CPU mem) to output render target
-			const FIntPoint SampleDim = Sample->GetDim();
-			const uint8* Data = (const uint8*)Sample->GetBuffer();
-			const FPixelFormatInfo& FormatInfo = GPixelFormats[MediaTextureResourceHelpers::GetPixelFormat(Sample)];
-
-			for (uint8 MipLevel = 0; MipLevel < SampleNumMips; ++MipLevel)
+			IMediaTextureSampleConverter::FConversionHints Hints;
+			Hints.NumMips = NumMips;
+			IMediaTextureSampleConverter* Converter = Sample->GetMediaTextureSampleConverter();
+			if (!Converter || !Converter->Convert(RHICmdList, RenderTargetTextureRHI, Hints))
 			{
-				uint32 Width = SampleDim.X >> MipLevel;
-				uint32 Height = SampleDim.Y >> MipLevel;
-				uint32 Stride = Sample->GetStride() >> MipLevel;
+				// Copy sample data (from CPU mem) to output render target
+				const FIntPoint SampleDim = Sample->GetDim();
+				const uint8* Data = (const uint8*)Sample->GetBuffer();
+				const FPixelFormatInfo& FormatInfo = GPixelFormats[MediaTextureResourceHelpers::GetPixelFormat(Sample)];
+
+				for (uint8 MipLevel = 0; MipLevel < SampleNumMips; ++MipLevel)
+				{
+					uint32 Width = SampleDim.X >> MipLevel;
+					uint32 Height = SampleDim.Y >> MipLevel;
+					uint32 Stride = Sample->GetStride() >> MipLevel;
 
-				FUpdateTextureRegion2D Region(0, 0, 0, 0, Width, Height);
-				RHIUpdateTexture2D(RenderTargetTextureRHI, MipLevel, Region, Stride, Data);
-				
-				uint32 HeightUnit = (Height + FormatInfo.BlockSizeY - 1) / FormatInfo.BlockSizeY;
-				Data += Stride * HeightUnit;
+					FUpdateTextureRegion2D Region(0, 0, 0, 0, Width, Height);
+					RHIUpdateTexture2D(RenderTargetTextureRHI, MipLevel, Region, Stride, Data);
+
+					uint32 HeightUnit = (Height + FormatInfo.BlockSizeY - 1) / FormatInfo.BlockSizeY;
+					Data += Stride * HeightUnit;
+				}
 			}
 
 			// Make sure resource is in SRV mode again
diff --git a/Engine/Source/Runtime/RenderCore/Private/MediaShaders.cpp b/Engine/Source/Runtime/RenderCore/Private/MediaShaders.cpp
--- a/Engine/Source/Runtime/RenderCore/Private/MediaShaders.cpp
+++ b/Engine/Source/Runtime/RenderCore/Private/MediaShaders.cpp
@@ -973,7 +973,8 @@ SHADER_PARAMETER(FMatrix44f, ColorTransform)
 SHADER_PARAMETER(FMatrix44f, CSTransform)
 SHADER_PARAMETER(uint32, EOTF)
 SHADER_PARAMETER(uint32, ToneMapMethod)
-SHADER_PARAMETER(FVector2f, UVScale)
+SHADER_PARAMETER(FLinearColor, UVRotation)
+SHADER_PARAMETER(FLinearColor, UVRotOffsetAndScale)
 SHADER_PARAMETER_TEXTURE(Texture2D, Texture)
 SHADER_PARAMETER_SAMPLER(SamplerState, Sampler)
 END_GLOBAL_SHADER_PARAMETER_STRUCT()
@@ -982,7 +983,8 @@ IMPLEMENT_GLOBAL_SHADER_PARAMETER_STRUCT(FVYUConvertUB, "VYUConvertUB");
 IMPLEMENT_SHADER_TYPE(, FVYUConvertPS, TEXT("/Engine/Private/MediaShaders.usf"), TEXT("VYUConvertPS"), SF_Pixel);
 
 
-void FVYUConvertPS::SetParameters(FRHIBatchedShaderParameters& BatchedParameters, TRefCountPtr<FRHITexture> VYUTexture, const FIntPoint& OutputDimensions, const FMatrix44f& ColorTransform, UE::Color::EEncoding Encoding, const FMatrix44f& CSTransform, MediaShaders::EToneMapMethod ToneMapMethod)
+void FVYUConvertPS::SetParameters(FRHIBatchedShaderParameters& BatchedParameters, TRefCountPtr<FRHITexture> VYUTexture, const FIntPoint& OutputDimensions, const FMatrix44f& ColorTransform, UE::Color::EEncoding Encoding, const FMatrix44f& CSTransform,
+	MediaShaders::EToneMapMethod ToneMapMethod, const FLinearColor& UVRotation, const FLinearColor& UVRotOffset)
 {
 	FVYUConvertUB UB;
 	{
@@ -990,7 +992,8 @@ void FVYUConvertPS::SetParameters(FRHIBatchedShaderParameters& BatchedParameters
 		UB.CSTransform = CSTransform;
 		UB.EOTF = static_cast<uint32>(Encoding);
 		UB.ToneMapMethod = static_cast<uint32>(ToneMapMethod);
-		UB.UVScale = FVector2f((float)OutputDimensions.X / (float)VYUTexture->GetSizeX(), (float)OutputDimensions.Y / (float)VYUTexture->GetSizeY());
+		UB.UVRotation = UVRotation;
+		UB.UVRotOffsetAndScale = FLinearColor(UVRotOffset.R, UVRotOffset.G, (float)OutputDimensions.X / (float)VYUTexture->GetSizeX(), (float)OutputDimensions.Y / (float)VYUTexture->GetSizeY());
 		UB.Texture = VYUTexture;
 		UB.Sampler = TStaticSamplerState<SF_Bilinear>::GetRHI();
 	}
diff --git a/Engine/Source/Runtime/RenderCore/Public/MediaShaders.h b/Engine/Source/Runtime/RenderCore/Public/MediaShaders.h
--- a/Engine/Source/Runtime/RenderCore/Public/MediaShaders.h
+++ b/Engine/Source/Runtime/RenderCore/Public/MediaShaders.h
@@ -740,5 +740,8 @@ public:
 		: FGlobalShader(Initializer)
 	{ }
 
-	RENDERCORE_API void SetParameters(FRHIBatchedShaderParameters& BatchedParameters, TRefCountPtr<FRHITexture> VYUTexture, const FIntPoint& OutputDimensions, const FMatrix44f& ColorTransform, UE::Color::EEncoding Encoding, const FMatrix44f& CSTransform, MediaShaders::EToneMapMethod ToneMapMethod = MediaShaders::EToneMapMethod::None);
+	RENDERCORE_API void SetParameters(FRHIBatchedShaderParameters& BatchedParameters, TRefCountPtr<FRHITexture> VYUTexture, const FIntPoint& OutputDimensions, const FMatrix44f& ColorTransform, UE::Color::EEncoding Encoding, const FMatrix44f& CSTransform, 
+		MediaShaders::EToneMapMethod ToneMapMethod = MediaShaders::EToneMapMethod::None,
+		const FLinearColor& UVRotation = FLinearColor(1.0f, 0.0f, 0.0f, 1.0f),
+		const FLinearColor& UVRotOffset = FLinearColor(0.0f, 0.0f, 0.0f, 0.0f));
 };
