Fixed CameraPlayer14/BitmapRendererImageReader NonUnity compilation.

diff --git a/Engine/Plugins/Media/AndroidCamera/Source/AndroidCamera/Private/Player/AndroidCameraTextureSample.h b/Engine/Plugins/Media/AndroidCamera/Source/AndroidCamera/Private/Player/AndroidCameraTextureSample.h
--- a/Engine/Plugins/Media/AndroidCamera/Source/AndroidCamera/Private/Player/AndroidCameraTextureSample.h
+++ b/Engine/Plugins/Media/AndroidCamera/Source/AndroidCamera/Private/Player/AndroidCameraTextureSample.h
@@ -12,6 +12,7 @@
 #include "RHI.h"
 #include "RHIUtilities.h"
 #include "Templates/SharedPointer.h"
+#include "PipelineStateCache.h"
 #include "RHIStaticStates.h"
 
 #if !WITH_ENGINE
@@ -409,10 +410,6 @@ private:
 			{
 				return RotateBinaryBufferTexture(RHICmdList, InDstTexture, Hints);
 			}
-			else
-			{
-				UE_LOG(LogAndroidCamera, Warning, TEXT("The binary buffer path for AndroidCameraTextureSample doesn't work without Vulkan."));
-			}
 		}
 
 		return false;
