Remove two imports not used.

diff --git a/Engine/Plugins/Media/AndroidCamera/Source/AndroidCamera/Java/src/com/epicgames/unreal/CameraPlayer14.java b/Engine/Plugins/Media/AndroidCamera/Source/AndroidCamera/Java/src/com/epicgames/unreal/CameraPlayer14.java
--- a/Engine/Plugins/Media/AndroidCamera/Source/AndroidCamera/Java/src/com/epicgames/unreal/CameraPlayer14.java
+++ b/Engine/Plugins/Media/AndroidCamera/Source/AndroidCamera/Java/src/com/epicgames/unreal/CameraPlayer14.java
@@ -30,8 +30,6 @@ import android.util.Range;
 import android.util.SparseIntArray;
 import android.view.Surface;
 
-import androidx.activity.result.ActivityResultLauncher;
-import androidx.activity.result.contract.ActivityResultContracts;
 import androidx.annotation.NonNull;
 import androidx.core.app.ActivityCompat;
 
