Fixed a Java exception when closing the CameraDevice during its initialization.
Fixed PlayOnOpen of AndroidCameraPlayer.

diff --git a/Engine/Plugins/Media/AndroidCamera/Source/AndroidCamera/Java/src/com/epicgames/unreal/CameraPlayer14.java b/Engine/Plugins/Media/AndroidCamera/Source/AndroidCamera/Java/src/com/epicgames/unreal/CameraPlayer14.java
--- a/Engine/Plugins/Media/AndroidCamera/Source/AndroidCamera/Java/src/com/epicgames/unreal/CameraPlayer14.java
+++ b/Engine/Plugins/Media/AndroidCamera/Source/AndroidCamera/Java/src/com/epicgames/unreal/CameraPlayer14.java
@@ -220,6 +220,8 @@ public class CameraPlayer14
 			return;
 		}
 
+		startBackgroundThread();
+
 		GameActivity.Log.debug("CameraPlayer14: frontId=" + CameraFrontId + ", backId=" + CameraBackId);
 	}
 
@@ -464,8 +466,6 @@ public class CameraPlayer14
 	{
 		reset();
 
-		startBackgroundThread();
-
 		try
 		{
 			if (!CameraSupported)
@@ -868,8 +868,6 @@ public class CameraPlayer14
 		audioTracks.clear();
 		videoTracks.clear();
 
-		stopBackgroundThread();
-
 		CameraState = CameraStates.RELEASED;
 	}
 
