Fixed the timing of firing MediaOpened in AndroidCameraPlayer.
The MediaOpened event will be fired at the correct time when the Java Camera2 device is prepared.
Enabling MediaPlayer.PlayOnOpen, binding the MediaOpened event, and connecting Play to OpenUrl are all working normally.

diff --git a/Engine/Plugins/Media/AndroidCamera/Source/AndroidCamera/Private/Player/AndroidCameraPlayer.cpp b/Engine/Plugins/Media/AndroidCamera/Source/AndroidCamera/Private/Player/AndroidCameraPlayer.cpp
--- a/Engine/Plugins/Media/AndroidCamera/Source/AndroidCamera/Private/Player/AndroidCameraPlayer.cpp
+++ b/Engine/Plugins/Media/AndroidCamera/Source/AndroidCamera/Private/Player/AndroidCameraPlayer.cpp
@@ -267,13 +267,6 @@ bool FAndroidCameraPlayer::Open(const FString& Url, const IMediaOptions* /*Optio
 	// prepare media
 	MediaUrl = JavaCameraPlayer->GetDataSource();
 
-	// notify listeners
-	if (!bOpenWithoutEvents)
-	{
-		EventSink.ReceiveMediaEvent(EMediaEvent::TracksChanged);
-		EventSink.ReceiveMediaEvent(EMediaEvent::MediaOpened);
-	}
-
 	CurrentState = EMediaState::Preparing;
 	return true;
 }
@@ -883,6 +876,11 @@ void FAndroidCameraPlayer::TickInput(FTimespan DeltaTime, FTimespan Timecode)
 			PauseHandle.Reset();
 		}
 
+		if (CurrentState == EMediaState::Preparing && JavaCameraPlayer->IsPrepared())
+		{
+			InitializePlayer();
+		}
+
 		return;
 	}
 
@@ -1046,6 +1044,13 @@ bool FAndroidCameraPlayer::InitializePlayer()
 
 	CurrentState = EMediaState::Stopped;
 
+	// notify listeners
+	if (!bOpenWithoutEvents)
+	{
+		EventSink.ReceiveMediaEvent(EMediaEvent::TracksChanged);
+		EventSink.ReceiveMediaEvent(EMediaEvent::MediaOpened);
+	}
+
 	return true;
 }
 
@@ -1627,35 +1632,51 @@ bool FAndroidCameraPlayer::SetRate(float Rate)
 
 		StopSetRateThread = false;
 
+		SetRateThreadArg* ThreadArg = new SetRateThreadArg();
+		ThreadArg->ThisPtr = StaticCastSharedRef<FAndroidCameraPlayer>(AsShared());
+
 		int ret = pthread_create(&DeferredSetRateThread, NULL, [] (void* arg) -> void* {
-			FAndroidCameraPlayer* Player = (FAndroidCameraPlayer*)arg;
-			Player->SetSetRateThreadStatus(true);
-			UE_LOG(LogAndroidCamera, Log, TEXT("Waiting for camera ready..."));
-			float ExitTime = FPlatformTime::Seconds() + 3.0f;
-			while (true)
+			SetRateThreadArg* Arg = static_cast<SetRateThreadArg*>(arg);
+			if (Arg->ThisPtr.IsValid())
 			{
-				if (FPlatformTime::Seconds() > ExitTime)
-				{
-					UE_LOG(LogAndroidCamera, Log, TEXT("Timed out after waiting for camera preparation for 3 seconds."));
-					break;
-				}
-				if (Player->IsStoppingSetRateThread())
-				{
-					UE_LOG(LogAndroidCamera, Log, TEXT("Stopped waiting for camera preparation."));
-					break;
-				}
-				if (Player->JavaCameraPlayer->IsPrepared())
+				Arg->ThisPtr->IsSetRateThreadRunning = true;
+
+				UE_LOG(LogAndroidCamera, Log, TEXT("Waiting for camera to become ready..."));
+
+				float ExitTime = FPlatformTime::Seconds() + 3.0f;
+				while (true)
 				{
-					UE_LOG(LogAndroidCamera, Log, TEXT("Camera is ready."));
-					Player->InitializePlayer();
-					Player->SetRate(1.0f);
-					break;
+					if (FPlatformTime::Seconds() > ExitTime)
+					{
+						UE_LOG(LogAndroidCamera, Log, TEXT("Timed out after waiting for camera preparation for 3 seconds."));
+						break;
+					}
+					if (Arg->ThisPtr->StopSetRateThread)
+					{
+						UE_LOG(LogAndroidCamera, Log, TEXT("Stopped waiting for camera preparation."));
+						break;
+					}
+					if ((Arg->ThisPtr->CurrentState != EMediaState::Preparing) &&
+						(Arg->ThisPtr->CurrentState != EMediaState::Error) &&
+						(Arg->ThisPtr->CurrentState != EMediaState::Closed))
+					{
+						UE_LOG(LogAndroidCamera, Log, TEXT("Camera is ready."));
+						Arg->ThisPtr->SetRate(1.0f);
+						break;
+					}
+					usleep(100);
 				}
-				usleep(100);
+
+				Arg->ThisPtr->IsSetRateThreadRunning = false;
 			}
-			Player->SetSetRateThreadStatus(false);
+			else
+			{
+				UE_LOG(LogAndroidCamera, Error, TEXT("The camera player object has been deleted."));
+			}
+
+			delete Arg;
 			return nullptr;
-		}, this);
+		}, ThreadArg);
 		
 		return (ret == 0);
 	}
diff --git a/Engine/Plugins/Media/AndroidCamera/Source/AndroidCamera/Private/Player/AndroidCameraPlayer.h b/Engine/Plugins/Media/AndroidCamera/Source/AndroidCamera/Private/Player/AndroidCameraPlayer.h
--- a/Engine/Plugins/Media/AndroidCamera/Source/AndroidCamera/Private/Player/AndroidCameraPlayer.h
+++ b/Engine/Plugins/Media/AndroidCamera/Source/AndroidCamera/Private/Player/AndroidCameraPlayer.h
@@ -31,6 +31,7 @@ class FAndroidCameraPlayer
 	, protected IMediaControls
 	, protected IMediaTracks
 	, protected IMediaView
+	, public TSharedFromThis<FAndroidCameraPlayer>
 {
 public:
 	
@@ -64,9 +65,6 @@ public:
 	virtual void TickFetch(FTimespan DeltaTime, FTimespan Timecode) override;
 	virtual void TickInput(FTimespan DeltaTime, FTimespan Timecode) override;
 
-	bool IsStoppingSetRateThread() { return StopSetRateThread; }
-	void SetSetRateThreadStatus(bool IsRunning) { IsSetRateThreadRunning = IsRunning; }
-
 protected:
 
 	/**
@@ -131,7 +129,7 @@ private:
 private:
 
 	/** Current player state. */
-	EMediaState CurrentState;
+	TAtomic<EMediaState> CurrentState;
 
 	/** The current time of the playback. */
 	FTimespan CurrentTime;
@@ -181,6 +179,13 @@ private:
 	/** Identify the game is in background */
 	bool IsInBackground;
 
+private:
+	/** The struct to pass a shared pointer of this object to DeferredSetRateThread */
+	struct SetRateThreadArg
+	{
+		TSharedPtr<FAndroidCameraPlayer> ThisPtr;
+	};
+
 	/** The thread of delaying the SetRate until the camera is prepared */
 	pthread_t DeferredSetRateThread;
 	TAtomic<bool> IsSetRateThreadRunning;
