Fixed printf.

diff --git a/Engine/Plugins/Media/AndroidCamera/Source/AndroidCamera/Private/Player/AndroidCameraPlayer.cpp b/Engine/Plugins/Media/AndroidCamera/Source/AndroidCamera/Private/Player/AndroidCameraPlayer.cpp
--- a/Engine/Plugins/Media/AndroidCamera/Source/AndroidCamera/Private/Player/AndroidCameraPlayer.cpp
+++ b/Engine/Plugins/Media/AndroidCamera/Source/AndroidCamera/Private/Player/AndroidCameraPlayer.cpp
@@ -1378,7 +1378,7 @@ bool FAndroidCameraPlayer::SelectTrack(EMediaTrackType TrackType, int32 TrackInd
 					}
 				}
 
-				UE_LOG(LogAndroidCamera, VeryVerbose, TEXT("Player %p: Enabling video"), this, TrackIndex);
+				UE_LOG(LogAndroidCamera, VeryVerbose, TEXT("Player %p: Enabling video"), this);
 				JavaCameraPlayer->SetVideoEnabled(true);
 			}
 
